Fix flickering of dynamic slate textures when first accessed

diff --git a/Engine/Source/Runtime/SlateRHIRenderer/Private/SlateRHIResourceManager.cpp b/Engine/Source/Runtime/SlateRHIRenderer/Private/SlateRHIResourceManager.cpp
--- a/Engine/Source/Runtime/SlateRHIRenderer/Private/SlateRHIResourceManager.cpp
+++ b/Engine/Source/Runtime/SlateRHIRenderer/Private/SlateRHIResourceManager.cpp
@@ -16,6 +16,7 @@ FSlateRHIResourceManager::FDynamicTextureResource::FDynamicTextureResource( FSla
 	, Proxy( new FSlateShaderResourceProxy )
 	, RHIRefTexture( ExistingTexture != NULL ? ExistingTexture : new FSlateTexture2DRHIRef(NULL,0,0) )
 {
+	Proxy->Resource = RHIRefTexture;
 }
 
 FSlateRHIResourceManager::FDynamicTextureResource::~FDynamicTextureResource()
@@ -393,8 +394,6 @@ TSharedRef< FSlateRHIResourceManager::FDynamicTextureResource > FSlateRHIResourc
 
 		}
 
-		TextureResource->Proxy->Resource = TextureResource->RHIRefTexture;
-
 		// Initialize and link the rendering resource
 		TextureResource->RHIRefTexture->InitResource();
 	})
