Add ensures to track down possible Slate dynamic brush alloc/dealloc outside of the game thread.

diff --git a/Engine/Source/Developer/StandaloneRenderer/Private/OpenGL/SlateOpenGLRenderer.cpp b/Engine/Source/Developer/StandaloneRenderer/Private/OpenGL/SlateOpenGLRenderer.cpp
--- a/Engine/Source/Developer/StandaloneRenderer/Private/OpenGL/SlateOpenGLRenderer.cpp
+++ b/Engine/Source/Developer/StandaloneRenderer/Private/OpenGL/SlateOpenGLRenderer.cpp
@@ -204,12 +204,14 @@ void FSlateOpenGLRenderer::UpdateFullscreenState( const TSharedRef<SWindow> InWi
 
 void FSlateOpenGLRenderer::ReleaseDynamicResource( const FSlateBrush& Brush )
 {
+	ensure( IsInGameThread() );
 	TextureManager->ReleaseDynamicTextureResource( Brush );
 }
 
 
 bool FSlateOpenGLRenderer::GenerateDynamicImageResource(FName ResourceName, uint32 Width, uint32 Height, const TArray< uint8 >& Bytes)
 {
+	ensure( IsInGameThread() );
 	return TextureManager->CreateDynamicTextureResource(ResourceName, Width, Height, Bytes) != NULL;
 }
 
diff --git a/Engine/Source/Developer/StandaloneRenderer/Private/Windows/D3D/SlateD3DRenderer.cpp b/Engine/Source/Developer/StandaloneRenderer/Private/Windows/D3D/SlateD3DRenderer.cpp
--- a/Engine/Source/Developer/StandaloneRenderer/Private/Windows/D3D/SlateD3DRenderer.cpp
+++ b/Engine/Source/Developer/StandaloneRenderer/Private/Windows/D3D/SlateD3DRenderer.cpp
@@ -191,11 +191,13 @@ void FSlateD3DRenderer::UpdateFullscreenState( const TSharedRef<SWindow> InWindo
 
 void FSlateD3DRenderer::ReleaseDynamicResource( const FSlateBrush& Brush )
 {
+	ensure( IsInGameThread() );
 	TextureManager->ReleaseDynamicTextureResource( Brush );
 }
 
 bool FSlateD3DRenderer::GenerateDynamicImageResource(FName ResourceName, uint32 Width, uint32 Height, const TArray< uint8 >& Bytes)
 {
+	ensure( IsInGameThread() );
 	return TextureManager->CreateDynamicTextureResource(ResourceName, Width, Height, Bytes) != NULL;
 }
 
diff --git a/Engine/Source/Runtime/SlateRHIRenderer/Private/SlateRHIRenderer.cpp b/Engine/Source/Runtime/SlateRHIRenderer/Private/SlateRHIRenderer.cpp
--- a/Engine/Source/Runtime/SlateRHIRenderer/Private/SlateRHIRenderer.cpp
+++ b/Engine/Source/Runtime/SlateRHIRenderer/Private/SlateRHIRenderer.cpp
@@ -906,12 +906,15 @@ void FSlateRHIRenderer::UnmapCrashTrackerBuffer()
 
 FIntPoint FSlateRHIRenderer::GenerateDynamicImageResource(const FName InTextureName)
 {
+	ensure( IsInGameThread() );
 	TSharedPtr<FSlateRHIResourceManager::FDynamicTextureResource> TextureResource = ResourceManager->MakeDynamicTextureResource(false, true, InTextureName.ToString(), InTextureName, NULL);
 	return TextureResource.IsValid() ? TextureResource->Proxy->ActualSize : FIntPoint( 0, 0 );
 }
 
 bool FSlateRHIRenderer::GenerateDynamicImageResource( FName ResourceName, uint32 Width, uint32 Height, const TArray< uint8 >& Bytes )
 {
+	ensure( IsInGameThread() );
+
 	TSharedPtr<FSlateRHIResourceManager::FDynamicTextureResource> TextureResource = ResourceManager->MakeDynamicTextureResource( ResourceName, Width, Height, Bytes );
 	return TextureResource.IsValid();
 }
@@ -978,6 +981,7 @@ void FSlateRHIRenderer::DisplayTextureAtlases()
 
 void FSlateRHIRenderer::ReleaseDynamicResource( const FSlateBrush& InBrush )
 {
+	ensure( IsInGameThread() );
 	ResourceManager->ReleaseDynamicResource( InBrush );
 }
 
