UE4 hopefully fix non unity builds

diff --git a/Engine/Source/Runtime/SlateRHIRenderer/Private/SlateShaders.h b/Engine/Source/Runtime/SlateRHIRenderer/Private/SlateShaders.h
--- a/Engine/Source/Runtime/SlateRHIRenderer/Private/SlateShaders.h
+++ b/Engine/Source/Runtime/SlateRHIRenderer/Private/SlateShaders.h
@@ -7,6 +7,7 @@
 #include "GlobalShader.h"
 #include "ShaderParameters.h"
 #include "UniformBuffer.h"
+#include "ShaderParameterUtils.h"
 
 extern uint32 GSlateShaderColorVisionDeficiencyType;
 
