SlateRHIRenderer: better PCH usage, renamed OVERRIDE to override

diff --git a/Engine/Source/Runtime/Launch/Launch.Build.cs b/Engine/Source/Runtime/Launch/Launch.Build.cs
--- a/Engine/Source/Runtime/Launch/Launch.Build.cs
+++ b/Engine/Source/Runtime/Launch/Launch.Build.cs
@@ -38,7 +38,7 @@ public class Launch : ModuleRules
         
         if( !UEBuildConfiguration.bBuildDedicatedServer)
         {
-            PrivateDependencyModuleNames.AddRange(
+            PrivateIncludePathModuleNames.AddRange(
                 new string[] {
 			        "SlateRHIRenderer",
 		        }
diff --git a/Engine/Source/Runtime/Launch/Private/LaunchEngineLoop.cpp b/Engine/Source/Runtime/Launch/Private/LaunchEngineLoop.cpp
--- a/Engine/Source/Runtime/Launch/Private/LaunchEngineLoop.cpp
+++ b/Engine/Source/Runtime/Launch/Private/LaunchEngineLoop.cpp
@@ -44,7 +44,7 @@
 	#include "Commandlets/Commandlet.h"
 
 #if !UE_SERVER
-	#include "SlateRHIRendererModule.h"
+	#include "ISlateRHIRendererModule.h"
 #endif
 
 	#include "MoviePlayer.h"
@@ -1262,7 +1262,7 @@ int32 FEngineLoop::PreInit( const TCHAR* CmdLine )
 #if !UE_SERVER// && !UE_EDITOR
 	if (!IsRunningDedicatedServer() && !IsRunningCommandlet())
 	{
-		TSharedRef<FSlateRenderer> SlateRenderer = FModuleManager::Get().GetModuleChecked<FSlateRHIRendererModule>("SlateRHIRenderer").CreateSlateRHIRenderer();
+		TSharedRef<FSlateRenderer> SlateRenderer = FModuleManager::Get().GetModuleChecked<ISlateRHIRendererModule>("SlateRHIRenderer").CreateSlateRHIRenderer();
 
 		// If Slate is being used, initialize the renderer after RHIInit 
 		FSlateApplication& CurrentSlateApp = FSlateApplication::Get();
@@ -1620,7 +1620,7 @@ void FEngineLoop::LoadPreInitModules()
 	if (!IsRunningDedicatedServer() )
 	{
 		// This needs to be loaded before InitializeShaderTypes is called
-		FModuleManager::Get().LoadModuleChecked<FSlateRHIRendererModule>("SlateRHIRenderer");
+		FModuleManager::Get().LoadModuleChecked<ISlateRHIRendererModule>("SlateRHIRenderer");
 	}
 #endif
 
diff --git a/Engine/Source/Runtime/SlateCore/Public/Rendering/SlateRenderer.h b/Engine/Source/Runtime/SlateCore/Public/Rendering/SlateRenderer.h
--- a/Engine/Source/Runtime/SlateCore/Public/Rendering/SlateRenderer.h
+++ b/Engine/Source/Runtime/SlateCore/Public/Rendering/SlateRenderer.h
@@ -1,7 +1,7 @@
 // Copyright 1998-2014 Epic Games, Inc. All Rights Reserved.
 
 /*=============================================================================
-	FSlateRenderer.h: Declares the FSlateRenderer class.
+	SlateRenderer.h: Declares the FSlateRenderer class.
 =============================================================================*/
 
 #pragma once
diff --git a/Engine/Source/Runtime/SlateRHIRenderer/Private/SlateRHIFontTexture.cpp b/Engine/Source/Runtime/SlateRHIRenderer/Private/SlateRHIFontTexture.cpp
--- a/Engine/Source/Runtime/SlateRHIRenderer/Private/SlateRHIFontTexture.cpp
+++ b/Engine/Source/Runtime/SlateRHIRenderer/Private/SlateRHIFontTexture.cpp
@@ -1,14 +1,15 @@
 // Copyright 1998-2014 Epic Games, Inc. All Rights Reserved.
 
+/*=============================================================================
+	SlateRHIFontTexture.cpp: Implements the FSlateFontTextureRHI class.
+=============================================================================*/
+
 #include "SlateRHIRendererPrivatePCH.h"
-#include "SlateRHIFontTexture.h"
 
 
 FSlateFontTextureRHI::FSlateFontTextureRHI( uint32 InWidth, uint32 InHeight )
 	: FSlateTexture2DRHIRef( InWidth, InHeight, PF_A8, NULL, TexCreate_Dynamic, true )
-{
-
-}
+{ }
 
 
 void FSlateFontTextureRHI::InitDynamicRHI()
diff --git a/Engine/Source/Runtime/SlateRHIRenderer/Private/SlateRHIFontTexture.h b/Engine/Source/Runtime/SlateRHIRenderer/Private/SlateRHIFontTexture.h
--- a/Engine/Source/Runtime/SlateRHIRenderer/Private/SlateRHIFontTexture.h
+++ b/Engine/Source/Runtime/SlateRHIRenderer/Private/SlateRHIFontTexture.h
@@ -19,8 +19,8 @@ public:
 	FSlateFontTextureRHI( uint32 InWidth, uint32 InHeight );
 
 	/** FRenderResource interface */
-	virtual void InitDynamicRHI() OVERRIDE;
-	virtual void ReleaseDynamicRHI() OVERRIDE;
+	virtual void InitDynamicRHI() override;
+	virtual void ReleaseDynamicRHI() override;
 private:
 	/** Temporary data stored between Release and InitDynamicRHI */
 	TArray<uint8> TempData;
@@ -38,9 +38,9 @@ public:
 	/**
 	 * FSlateFontAtlas interface 
 	 */
-	virtual class FSlateShaderResource* GetTexture()  OVERRIDE { return FontTexture; }
-	virtual void ConditionalUpdateTexture()  OVERRIDE;
-	virtual void ReleaseResources() OVERRIDE;
+	virtual class FSlateShaderResource* GetTexture()  override { return FontTexture; }
+	virtual void ConditionalUpdateTexture()  override;
+	virtual void ReleaseResources() override;
 private:
 	FSlateTexture2DRHIRef* FontTexture;
 
diff --git a/Engine/Source/Runtime/SlateRHIRenderer/Private/SlateRHIRenderer.cpp b/Engine/Source/Runtime/SlateRHIRenderer/Private/SlateRHIRenderer.cpp
--- a/Engine/Source/Runtime/SlateRHIRenderer/Private/SlateRHIRenderer.cpp
+++ b/Engine/Source/Runtime/SlateRHIRenderer/Private/SlateRHIRenderer.cpp
@@ -1,17 +1,10 @@
 // Copyright 1998-2014 Epic Games, Inc. All Rights Reserved.
 
+/*=============================================================================
+	SlateRHIRenderer.cpp: Implements the FSlateRenderer class and related classes.
+=============================================================================*/
 
 #include "SlateRHIRendererPrivatePCH.h"
-#include "SlateRHIRenderer.h"
-#include "SlateRHIFontTexture.h"
-#include "SlateRHIResourceManager.h"
-#include "SlateCore.h"
-#include "RHIStaticStates.h"
-#include "SlateRHIRenderingPolicy.h"
-#include "ShaderParameterUtils.h"
-#include "Runtime/Engine/Public/ScreenRendering.h"
-#include "Runtime/Engine/Public/ShaderCompiler.h"
-#include "SlateShaders.h"
 
 DECLARE_CYCLE_STAT(TEXT("Map Staging Buffer"),STAT_MapStagingBuffer,STATGROUP_CrashTracker);
 DECLARE_CYCLE_STAT(TEXT("Generate Capture Buffer"),STAT_GenerateCaptureBuffer,STATGROUP_CrashTracker);
@@ -23,6 +16,7 @@ namespace CrashTrackerConstants
 	static const float ScreenScaling = 0.5f;
 }
 
+
 // Defines the maximum size that a slate viewport will create
 #define MAX_VIEWPORT_SIZE 16384
 
@@ -172,7 +166,7 @@ public:
 	}
 
 
-	virtual TSharedRef<FSlateFontAtlas> CreateFontAtlas() const OVERRIDE
+	virtual TSharedRef<FSlateFontAtlas> CreateFontAtlas() const override
 	{
 		return MakeShareable( new FSlateFontAtlasRHI( AtlasSize, AtlasSize ) );
 	}
diff --git a/Engine/Source/Runtime/SlateRHIRenderer/Private/SlateRHIRenderer.h b/Engine/Source/Runtime/SlateRHIRenderer/Private/SlateRHIRenderer.h
--- a/Engine/Source/Runtime/SlateRHIRenderer/Private/SlateRHIRenderer.h
+++ b/Engine/Source/Runtime/SlateRHIRenderer/Private/SlateRHIRenderer.h
@@ -122,22 +122,22 @@ public:
 	~FSlateRHIRenderer();
 
 	/** FSlateRenderer interface */
-	virtual void Initialize() OVERRIDE;
-	virtual void Destroy() OVERRIDE;
-	virtual FSlateDrawBuffer& GetDrawBuffer() OVERRIDE;
-	virtual void OnWindowDestroyed( const TSharedRef<SWindow>& InWindow ) OVERRIDE;
-	virtual void RequestResize( const TSharedPtr<SWindow>& Window, uint32 NewWidth, uint32 NewHeight ) OVERRIDE;
-	virtual void CreateViewport( const TSharedRef<SWindow> Window ) OVERRIDE;
-	virtual void UpdateFullscreenState( const TSharedRef<SWindow> Window, uint32 OverrideResX, uint32 OverrideResY ) OVERRIDE;
-	virtual void DrawWindows( FSlateDrawBuffer& InWindowDrawBuffer ) OVERRIDE;
-	virtual void DrawWindows() OVERRIDE;
-	virtual void FlushCommands() const OVERRIDE;
-	virtual void Sync() const OVERRIDE;
-	virtual void ReleaseDynamicResource( const FSlateBrush& InBrush ) OVERRIDE;
-	virtual FIntPoint GenerateDynamicImageResource(const FName InTextureName) OVERRIDE;
-	virtual bool GenerateDynamicImageResource( FName ResourceName, uint32 Width, uint32 Height, const TArray< uint8 >& Bytes ) OVERRIDE;
-	virtual void* GetViewportResource( const SWindow& Window ) OVERRIDE;
-	virtual void SetColorVisionDeficiencyType( uint32 Type ) OVERRIDE;
+	virtual void Initialize() override;
+	virtual void Destroy() override;
+	virtual FSlateDrawBuffer& GetDrawBuffer() override;
+	virtual void OnWindowDestroyed( const TSharedRef<SWindow>& InWindow ) override;
+	virtual void RequestResize( const TSharedPtr<SWindow>& Window, uint32 NewWidth, uint32 NewHeight ) override;
+	virtual void CreateViewport( const TSharedRef<SWindow> Window ) override;
+	virtual void UpdateFullscreenState( const TSharedRef<SWindow> Window, uint32 OverrideResX, uint32 OverrideResY ) override;
+	virtual void DrawWindows( FSlateDrawBuffer& InWindowDrawBuffer ) override;
+	virtual void DrawWindows() override;
+	virtual void FlushCommands() const override;
+	virtual void Sync() const override;
+	virtual void ReleaseDynamicResource( const FSlateBrush& InBrush ) override;
+	virtual FIntPoint GenerateDynamicImageResource(const FName InTextureName) override;
+	virtual bool GenerateDynamicImageResource( FName ResourceName, uint32 Width, uint32 Height, const TArray< uint8 >& Bytes ) override;
+	virtual void* GetViewportResource( const SWindow& Window ) override;
+	virtual void SetColorVisionDeficiencyType( uint32 Type ) override;
 
 	/** Draws windows from a FSlateDrawBuffer on the render thread */
 	void DrawWindow_RenderThread( const FSlateRHIRenderer::FViewportInfo& ViewportInfo, const FSlateWindowElementList& WindowElementList, bool bLockToVsync );
@@ -150,16 +150,16 @@ public:
 	virtual void CopyWindowsToDrawBuffer(const TArray<FString>& KeypressBuffer);
 	
 	/** Allows and disallows access to the crash tracker buffer data on the CPU */
-	virtual void MapCrashTrackerBuffer(void** OutImageData, int32* OutWidth, int32* OutHeight) OVERRIDE;
-	virtual void UnmapCrashTrackerBuffer() OVERRIDE;
+	virtual void MapCrashTrackerBuffer(void** OutImageData, int32* OutWidth, int32* OutHeight) override;
+	virtual void UnmapCrashTrackerBuffer() override;
 
 	/**
 	 * Reloads texture resources from disk                   
 	 */
-	virtual void ReloadTextureResources() OVERRIDE;
+	virtual void ReloadTextureResources() override;
 
 
-	virtual void LoadStyleResources( const ISlateStyle& Style ) OVERRIDE;
+	virtual void LoadStyleResources( const ISlateStyle& Style ) override;
 
 	/**
 	 * Creates a window with an atlas visualizer inside it
@@ -183,7 +183,7 @@ public:
 	 */
 	virtual void InvalidateAllViewports();
 
-	virtual void PrepareToTakeScreenshot(const FIntRect& Rect, TArray<FColor>* OutColorData) OVERRIDE;
+	virtual void PrepareToTakeScreenshot(const FIntRect& Rect, TArray<FColor>* OutColorData) override;
 
 private:
 	/** Loads all known textures from Slate styles */
diff --git a/Engine/Source/Runtime/SlateRHIRenderer/Private/SlateRHIRendererModule.cpp b/Engine/Source/Runtime/SlateRHIRenderer/Private/SlateRHIRendererModule.cpp
--- a/Engine/Source/Runtime/SlateRHIRenderer/Private/SlateRHIRendererModule.cpp
+++ b/Engine/Source/Runtime/SlateRHIRenderer/Private/SlateRHIRendererModule.cpp
@@ -1,22 +1,33 @@
 // Copyright 1998-2014 Epic Games, Inc. All Rights Reserved.
 
+/*=============================================================================
+	SlateRHIRendererModule.cpp: Implements the FSlateRHIRendererModule class.
+=============================================================================*/
+
 #include "SlateRHIRendererPrivatePCH.h"
-#include "SlateRHIRendererModule.h"
-#include "SlateRHIRenderer.h"
 
-void FSlateRHIRendererModule::StartupModule()
+
+/**
+ * Implements the Slate RHI Renderer module.
+ */
+class FSlateRHIRendererModule
+	: public ISlateRHIRendererModule
 {
+public:
 
-}
+	// Begin ISlateRHIRendererModule interface
 
-void FSlateRHIRendererModule::ShutdownModule()
-{
+	virtual TSharedRef<FSlateRenderer> CreateSlateRHIRenderer( ) override
+	{
+		return TSharedRef<FSlateRHIRenderer>( new FSlateRHIRenderer );
+	}
 
-}
+	virtual void StartupModule( ) override { }
+	
+	virtual void ShutdownModule( ) override { }
+
+	// End ISlateRHIRendererModule interface
+};
 
-TSharedRef<FSlateRenderer> FSlateRHIRendererModule::CreateSlateRHIRenderer()
-{
-	return TSharedRef<FSlateRHIRenderer>( new FSlateRHIRenderer );
-}
 
 IMPLEMENT_MODULE( FSlateRHIRendererModule, SlateRHIRenderer ) 
diff --git a/Engine/Source/Runtime/SlateRHIRenderer/Private/SlateRHIRendererPrivatePCH.h b/Engine/Source/Runtime/SlateRHIRenderer/Private/SlateRHIRendererPrivatePCH.h
--- a/Engine/Source/Runtime/SlateRHIRenderer/Private/SlateRHIRendererPrivatePCH.h
+++ b/Engine/Source/Runtime/SlateRHIRenderer/Private/SlateRHIRendererPrivatePCH.h
@@ -1,9 +1,35 @@
 // Copyright 1998-2014 Epic Games, Inc. All Rights Reserved.
 
+/*=============================================================================
+	SlateRHIRendererPrivatePCH.h: Pre-compiled header file for the SlateRHIRenderer module.
+=============================================================================*/
+
 #pragma once
 
+#include "ISlateRHIRendererModule.h"
+
+
+/* Private Dependencies
+ *****************************************************************************/
+
 #include "Slate.h"
+#include "SlateCore.h"
 #include "RenderingCommon.h"
 #include "Runtime/RHI/Public/RHI.h"
 #include "Runtime/Engine/Public/Engine.h"
+#include "Runtime/Engine/Public/ScreenRendering.h"
+#include "Runtime/Engine/Public/ShaderCompiler.h"
 #include "Runtime/Engine/Public/Slate/SlateTextures.h"
+#include "RHIStaticStates.h"
+#include "ShaderParameterUtils.h"
+
+
+/* Private Includes
+ *****************************************************************************/
+
+#include "SlateRHIFontTexture.h"
+#include "SlateRHIResourceManager.h"
+#include "SlateRHIRenderingPolicy.h"
+#include "SlateRHITextureAtlas.h"
+#include "SlateShaders.h"
+#include "SlateRHIRenderer.h"
diff --git a/Engine/Source/Runtime/SlateRHIRenderer/Private/SlateRHIRenderingPolicy.cpp b/Engine/Source/Runtime/SlateRHIRenderer/Private/SlateRHIRenderingPolicy.cpp
--- a/Engine/Source/Runtime/SlateRHIRenderer/Private/SlateRHIRenderingPolicy.cpp
+++ b/Engine/Source/Runtime/SlateRHIRenderer/Private/SlateRHIRenderingPolicy.cpp
@@ -1,12 +1,14 @@
 // Copyright 1998-2014 Epic Games, Inc. All Rights Reserved.
 
+/*=============================================================================
+	SlateRHIRenderingPolicy.cpp: Implements the FSlateRHIRenderingPolicy class and related classes.
+=============================================================================*/
+
 #include "SlateRHIRendererPrivatePCH.h"
 #include "RenderingPolicy.h"
-#include "SlateRHIRenderingPolicy.h"
-#include "SlateShaders.h"
-#include "SlateRHIResourceManager.h"
 #include "RHIStaticStates.h"
 
+
 static EPrimitiveType GetRHIPrimitiveType( ESlateDrawPrimitive::Type SlateType )
 {
 	switch( SlateType )
diff --git a/Engine/Source/Runtime/SlateRHIRenderer/Private/SlateRHIRenderingPolicy.h b/Engine/Source/Runtime/SlateRHIRenderer/Private/SlateRHIRenderingPolicy.h
--- a/Engine/Source/Runtime/SlateRHIRenderer/Private/SlateRHIRenderingPolicy.h
+++ b/Engine/Source/Runtime/SlateRHIRenderer/Private/SlateRHIRenderingPolicy.h
@@ -103,7 +103,7 @@ public:
 	{
 		RenderTargetTextureRHI = InRenderTargetTexture;
 	}
-	virtual FIntPoint GetSizeXY() const OVERRIDE { return SizeXY; }
+	virtual FIntPoint GetSizeXY() const override { return SizeXY; }
 private:
 	FIntPoint SizeXY;
 };
@@ -114,11 +114,11 @@ public:
 	FSlateRHIRenderingPolicy( TSharedPtr<FSlateFontCache> InFontCache, TSharedRef<FSlateRHIResourceManager> InTextureManager );
 	~FSlateRHIRenderingPolicy();
 	
-	virtual void UpdateBuffers( const FSlateWindowElementList& WindowElementList ) OVERRIDE;
+	virtual void UpdateBuffers( const FSlateWindowElementList& WindowElementList ) override;
 	virtual void DrawElements( const FIntPoint& InViewportSize, FSlateRenderTarget& BackBuffer, const FMatrix& ViewProjectionMatrix, const TArray<FSlateRenderBatch>& RenderBatches );
 
 	class FSlateShaderResource* GetViewportResource( const ISlateViewport* InViewportInterface );
-	virtual class FSlateShaderResourceProxy* GetTextureResource( const FSlateBrush& Brush ) OVERRIDE;
+	virtual class FSlateShaderResourceProxy* GetTextureResource( const FSlateBrush& Brush ) override;
 	TSharedPtr<FSlateFontCache>& GetFontCache() { return FontCache; }
 	
 	void InitResources();
diff --git a/Engine/Source/Runtime/SlateRHIRenderer/Private/SlateRHIResourceManager.cpp b/Engine/Source/Runtime/SlateRHIRenderer/Private/SlateRHIResourceManager.cpp
--- a/Engine/Source/Runtime/SlateRHIRenderer/Private/SlateRHIResourceManager.cpp
+++ b/Engine/Source/Runtime/SlateRHIRenderer/Private/SlateRHIResourceManager.cpp
@@ -1,10 +1,12 @@
 // Copyright 1998-2014 Epic Games, Inc. All Rights Reserved.
 
+/*=============================================================================
+	SlateRHIRenderingPolicy.cpp: Implements the FSlateRHIResourceManager class.
+=============================================================================*/
+
 #include "SlateRHIRendererPrivatePCH.h"
-#include "SlateRHIRenderer.h"
-#include "SlateRHIResourceManager.h"
 #include "ImageWrapper.h"
-#include "SlateRHITextureAtlas.h"
+
 
 typedef TMap<FName,FSlateTexture2DRHIRef*> FSlateTextureMap;
 
@@ -609,9 +611,9 @@ public:
 		, SelectedAtlasPage( NULL )
 	{}
 
-	virtual FIntPoint GetSize() const OVERRIDE { return FIntPoint(1024,1024); }
+	virtual FIntPoint GetSize() const override { return FIntPoint(1024,1024); }
 
-	virtual FSlateShaderResource* GetViewportRenderTargetTexture() const OVERRIDE
+	virtual FSlateShaderResource* GetViewportRenderTargetTexture() const override
 	{
 		if( SelectedAtlasPage.IsValid() )
 		{
@@ -623,7 +625,7 @@ public:
 		}
 	}
 
-	virtual bool RequiresVsync() const OVERRIDE { return false; }
+	virtual bool RequiresVsync() const override { return false; }
 
 	TSharedRef<SWidget> MakeVisualizerWidget()
 	{
diff --git a/Engine/Source/Runtime/SlateRHIRenderer/Private/SlateRHIResourceManager.h b/Engine/Source/Runtime/SlateRHIRenderer/Private/SlateRHIResourceManager.h
--- a/Engine/Source/Runtime/SlateRHIRenderer/Private/SlateRHIResourceManager.h
+++ b/Engine/Source/Runtime/SlateRHIRenderer/Private/SlateRHIResourceManager.h
@@ -149,7 +149,7 @@ private:
 	FSlateShaderResourceProxy* GetDynamicTextureResource( const FSlateBrush& InBrush );
 
 	/** FGCObject */
-	void AddReferencedObjects( FReferenceCollector& Collector ) OVERRIDE;
+	void AddReferencedObjects( FReferenceCollector& Collector ) override;
 
 private:
 	/** Map of all active dynamic texture objects being used for brush resources */
diff --git a/Engine/Source/Runtime/SlateRHIRenderer/Private/SlateRHITextureAtlas.cpp b/Engine/Source/Runtime/SlateRHIRenderer/Private/SlateRHITextureAtlas.cpp
--- a/Engine/Source/Runtime/SlateRHIRenderer/Private/SlateRHITextureAtlas.cpp
+++ b/Engine/Source/Runtime/SlateRHIRenderer/Private/SlateRHITextureAtlas.cpp
@@ -1,35 +1,69 @@
 // Copyright 1998-2014 Epic Games, Inc. All Rights Reserved.
 
+/*=============================================================================
+	SlateRHITextureAtlas.cpp: Implements the FSlateTextureAtlasRHI class.
+=============================================================================*/
+
 #include "SlateRHIRendererPrivatePCH.h"
-#include "SlateRHITextureAtlas.h"
 
 
+/* FSlateTextureAtlasRHI structors
+ *****************************************************************************/
+
 FSlateTextureAtlasRHI::FSlateTextureAtlasRHI( uint32 Width, uint32 Height, uint32 StrideBytes, uint32 Padding )
-	: FSlateTextureAtlas( Width, Height, StrideBytes, Padding )
-	, AtlasTexture( new FSlateTexture2DRHIRef( Width, Height, PF_B8G8R8A8, NULL, TexCreate_SRGB, true ) ) 
-{
-}
+	: FSlateTextureAtlas(Width, Height, StrideBytes, Padding)
+	, AtlasTexture(new FSlateTexture2DRHIRef(Width, Height, PF_B8G8R8A8, NULL, TexCreate_SRGB, true)) 
+{ }
+
 
-FSlateTextureAtlasRHI::~FSlateTextureAtlasRHI()
+FSlateTextureAtlasRHI::~FSlateTextureAtlasRHI( )
 {
-	if( AtlasTexture )
+	if (AtlasTexture != nullptr)
 	{
 		delete AtlasTexture;
 	}
 }
 
 
-void FSlateTextureAtlasRHI::ReleaseAtlasTexture()
+/* FSlateTextureAtlasRHI interface
+ *****************************************************************************/
+
+void FSlateTextureAtlasRHI::ReleaseAtlasTexture( )
 {
 	bNeedsUpdate = false;
-	BeginReleaseResource( AtlasTexture );
+	BeginReleaseResource(AtlasTexture);
+}
+
+
+void FSlateTextureAtlasRHI::UpdateTexture_RenderThread( FSlateTextureData* RenderThreadData )
+{
+	check(IsInRenderingThread());
+
+	if (!AtlasTexture->IsInitialized())
+	{
+		AtlasTexture->InitResource();
+	}
+
+	check(AtlasTexture->IsInitialized());
+
+	uint32 Stride;
+	uint8* TempData = (uint8*)RHILockTexture2D(AtlasTexture->GetTypedResource(), 0, RLM_WriteOnly, Stride, false);
+
+	FMemory::Memcpy(TempData, RenderThreadData->GetRawBytes().GetData(), RenderThreadData->GetStride() * RenderThreadData->GetWidth() * RenderThreadData->GetHeight());
+	RHIUnlockTexture2D(AtlasTexture->GetTypedResource(), 0, false);
+
+	delete RenderThreadData;
 }
 
-void FSlateTextureAtlasRHI::ConditionalUpdateTexture()
+
+/* FSlateTextureAtlas overrides
+ *****************************************************************************/
+
+void FSlateTextureAtlasRHI::ConditionalUpdateTexture( )
 {
-	check( IsThreadSafeForSlateRendering() );
+	check(IsThreadSafeForSlateRendering());
 
-	if( bNeedsUpdate )
+	if (bNeedsUpdate)
 	{
 		// Copy the game thread data. This is deleted on the render thread
 		FSlateTextureData* RenderThreadData = new FSlateTextureData( AtlasWidth, AtlasHeight, Stride, AtlasData ); 
@@ -42,25 +76,4 @@ void FSlateTextureAtlasRHI::ConditionalUpdateTexture()
 
 		bNeedsUpdate = false;
 	}
-
-
-}
-
-void FSlateTextureAtlasRHI::UpdateTexture_RenderThread( FSlateTextureData* RenderThreadData )
-{
-	check( IsInRenderingThread() );
-
-	if( !AtlasTexture->IsInitialized() )
-	{
-		AtlasTexture->InitResource();
-	}
-
-	check( AtlasTexture->IsInitialized() );
-
-	uint32 Stride;
-	uint8* TempData = (uint8*)RHILockTexture2D( AtlasTexture->GetTypedResource(), 0, RLM_WriteOnly, Stride, false );
-	FMemory::Memcpy( TempData, RenderThreadData->GetRawBytes().GetData(), RenderThreadData->GetStride()*RenderThreadData->GetWidth()*RenderThreadData->GetHeight() );
-	RHIUnlockTexture2D( AtlasTexture->GetTypedResource(),0,false );
-
-	delete RenderThreadData;
 }
diff --git a/Engine/Source/Runtime/SlateRHIRenderer/Private/SlateRHITextureAtlas.h b/Engine/Source/Runtime/SlateRHIRenderer/Private/SlateRHITextureAtlas.h
--- a/Engine/Source/Runtime/SlateRHIRenderer/Private/SlateRHITextureAtlas.h
+++ b/Engine/Source/Runtime/SlateRHIRenderer/Private/SlateRHITextureAtlas.h
@@ -1,32 +1,69 @@
 // Copyright 1998-2014 Epic Games, Inc. All Rights Reserved.
 
+/*=============================================================================
+	SlateRHITextureAtlas.h: Declares the FSlateTextureAtlasRHI class.
+=============================================================================*/
+
 #pragma once
 
+
 /**
- * Represents a texture atlas for use with RHI
+ * Represents a texture atlas for use with RHI.
  */
-class FSlateTextureAtlasRHI : public FSlateTextureAtlas
+class FSlateTextureAtlasRHI
+	: public FSlateTextureAtlas
 {
 public:
+
+	/**
+	 * Creates and initializes a new instance.
+	 *
+	 * @param Width
+	 * @param Height
+	 * @param StrideBytes
+	 * @param Padding
+	 */
 	FSlateTextureAtlasRHI( uint32 Width, uint32 Height, uint32 SrideBytes, uint32 Padding );
-	~FSlateTextureAtlasRHI();
 
 	/**
-	 * Releases rendering resources from the texture
+	 * Destructor.
+	 */
+	~FSlateTextureAtlasRHI( );
+
+public:
+
+	/**
+	 * Gets the atlas' underlying texture resource.
+	 *
+	 * @return The texture resource.
+	 */
+	FSlateTexture2DRHIRef* GetAtlasTexture( ) const
+	{
+		return AtlasTexture;
+	}
+
+	/**
+	 * Releases rendering resources from the texture.
 	 */	
-	void ReleaseAtlasTexture();
-	
-	/** @return The texture resource for rendering */
-	FSlateTexture2DRHIRef* GetAtlasTexture() const { return AtlasTexture; }
+	void ReleaseAtlasTexture( );
 	
-	/** FSlateTextureAtlas interface */
-	virtual void ConditionalUpdateTexture();
-
 	/**
-	 * Updates the texture on the render thread
+	 * Updates the texture on the render thread.
+	 *
+	 * @param RenderThreadData
 	 */
 	void UpdateTexture_RenderThread( FSlateTextureData* RenderThreadData );
+
+public:
+
+	// Begin FSlateTextureAtlas overrides.
+
+	virtual void ConditionalUpdateTexture( );
+
+	// End FSlateTextureAtlas interface.
+
 private:
+
 	/** The texture rendering resource */
 	FSlateTexture2DRHIRef* AtlasTexture;
-};
\ No newline at end of file
+};
diff --git a/Engine/Source/Runtime/SlateRHIRenderer/Private/SlateShaders.cpp b/Engine/Source/Runtime/SlateRHIRenderer/Private/SlateShaders.cpp
--- a/Engine/Source/Runtime/SlateRHIRenderer/Private/SlateShaders.cpp
+++ b/Engine/Source/Runtime/SlateRHIRenderer/Private/SlateShaders.cpp
@@ -1,10 +1,11 @@
 // Copyright 1998-2014 Epic Games, Inc. All Rights Reserved.
 
+/*=============================================================================
+	SlateShaders.cpp: Implements shader related classes for the Slate RHI Renderer.
+=============================================================================*/
 
 #include "SlateRHIRendererPrivatePCH.h"
-#include "SlateRHIRenderingPolicy.h"
-#include "SlateShaders.h"
-#include "ShaderParameterUtils.h"
+
 
 /** Flag to determine if we are running with a color vision deficiency shader on */
 uint32 GSlateShaderColorVisionDeficiencyType = 0;
diff --git a/Engine/Source/Runtime/SlateRHIRenderer/Public/Interfaces/ISlateRHIRendererModule.h b/Engine/Source/Runtime/SlateRHIRenderer/Public/Interfaces/ISlateRHIRendererModule.h
new file mode 100644
--- /dev/null
+++ b/Engine/Source/Runtime/SlateRHIRenderer/Public/Interfaces/ISlateRHIRendererModule.h
@@ -0,0 +1,27 @@
+// Copyright 1998-2014 Epic Games, Inc. All Rights Reserved.
+
+/*=============================================================================
+	ISlateRHIRendererModule.h: Declares the ISlateRHIRendererModule interface.
+=============================================================================*/
+
+#pragma once
+
+#include "ModuleInterface.h"
+#include "SlateCore.h"
+
+
+/**
+ * Interface for the Slate RHI Renderer module.
+ */
+class ISlateRHIRendererModule
+	: public IModuleInterface
+{
+public:
+
+	/**
+	 * Creates a Slate RHI renderer.
+	 *
+	 * @return A new renderer.
+	 */
+	virtual TSharedRef<FSlateRenderer> CreateSlateRHIRenderer( ) = 0;
+};
diff --git a/Engine/Source/Runtime/SlateRHIRenderer/Public/SlateRHIRendererModule.h b/Engine/Source/Runtime/SlateRHIRenderer/Public/SlateRHIRendererModule.h
deleted file mode 100644
