Fixed Slate dynamic UTexture resource brushes with the same name failing to be unique.

diff --git a/Engine/Source/Runtime/SlateRHIRenderer/Private/SlateRHIRenderer.cpp b/Engine/Source/Runtime/SlateRHIRenderer/Private/SlateRHIRenderer.cpp
--- a/Engine/Source/Runtime/SlateRHIRenderer/Private/SlateRHIRenderer.cpp
+++ b/Engine/Source/Runtime/SlateRHIRenderer/Private/SlateRHIRenderer.cpp
@@ -903,16 +903,16 @@ void FSlateRHIRenderer::UnmapCrashTrackerBuffer()
 
 FIntPoint FSlateRHIRenderer::GenerateDynamicImageResource(const FName InTextureName)
 {
-	ensure( IsInGameThread() );
-	TSharedPtr<FSlateRHIResourceManager::FDynamicTextureResource> TextureResource = ResourceManager->MakeDynamicTextureResource(false, true, InTextureName.ToString(), InTextureName, NULL);
+	check( IsInGameThread() );
+	TSharedPtr<FDynamicTextureResource> TextureResource = ResourceManager->MakeDynamicTextureResource(false, true, InTextureName.ToString(), InTextureName, NULL);
 	return TextureResource.IsValid() ? TextureResource->Proxy->ActualSize : FIntPoint( 0, 0 );
 }
 
 bool FSlateRHIRenderer::GenerateDynamicImageResource( FName ResourceName, uint32 Width, uint32 Height, const TArray< uint8 >& Bytes )
 {
-	ensure( IsInGameThread() );
+	check( IsInGameThread() );
 
-	TSharedPtr<FSlateRHIResourceManager::FDynamicTextureResource> TextureResource = ResourceManager->MakeDynamicTextureResource( ResourceName, Width, Height, Bytes );
+	TSharedPtr<FDynamicTextureResource> TextureResource = ResourceManager->MakeDynamicTextureResource( ResourceName, Width, Height, Bytes );
 	return TextureResource.IsValid();
 }
 
diff --git a/Engine/Source/Runtime/SlateRHIRenderer/Private/SlateRHIResourceManager.cpp b/Engine/Source/Runtime/SlateRHIRenderer/Private/SlateRHIResourceManager.cpp
--- a/Engine/Source/Runtime/SlateRHIRenderer/Private/SlateRHIResourceManager.cpp
+++ b/Engine/Source/Runtime/SlateRHIRenderer/Private/SlateRHIResourceManager.cpp
@@ -4,32 +4,110 @@
 #include "ImageWrapper.h"
 
 
-typedef TMap<FName,FSlateTexture2DRHIRef*> FSlateTextureMap;
 
-TSharedPtr<FSlateRHIResourceManager::FDynamicTextureResource> FSlateRHIResourceManager::FDynamicTextureResource::NullResource = MakeShareable( new FDynamicTextureResource( NULL ) );
+TSharedPtr<FDynamicTextureResource> FDynamicTextureResource::NullResource = MakeShareable( new FDynamicTextureResource( NULL ) );
 
 
-FSlateRHIResourceManager::FDynamicTextureResource::FDynamicTextureResource( FSlateTexture2DRHIRef* ExistingTexture )
+FDynamicTextureResource::FDynamicTextureResource(FSlateTexture2DRHIRef* ExistingTexture)
 	: TextureObject( NULL )
-	, Proxy( new FSlateShaderResourceProxy )
-	, RHIRefTexture( ExistingTexture != NULL ? ExistingTexture : new FSlateTexture2DRHIRef(NULL,0,0) )
+	, Proxy(new FSlateShaderResourceProxy)
+	, RHIRefTexture(ExistingTexture != NULL ? ExistingTexture : new FSlateTexture2DRHIRef(NULL, 0, 0))
 {
 	Proxy->Resource = RHIRefTexture;
 }
 
-FSlateRHIResourceManager::FDynamicTextureResource::~FDynamicTextureResource()
+FDynamicTextureResource::~FDynamicTextureResource()
 {
-	if( Proxy )
+	if (Proxy)
 	{
 		delete Proxy;
 	}
 
-	if( RHIRefTexture )
+	if (RHIRefTexture)
 	{
 		delete RHIRefTexture;
 	}
 }
 
+
+TSharedPtr<FDynamicTextureResource> FDynamicResourceMap::GetResource( FName ResourceName, UObject* ResourceObject ) const
+{
+	if( ResourceObject )
+	{
+		return DynamicResourceObjectMap.FindRef( ResourceObject );
+	}
+	else
+	{
+		return DynamicNativeTextureMap.FindRef( ResourceName );
+	}
+}
+
+
+void FDynamicResourceMap::AddResource( FName ResourceName, UObject* ResourceObject, TSharedRef<FDynamicTextureResource> InResource )
+{
+	if( ResourceObject )
+	{
+		DynamicResourceObjectMap.Add( ResourceObject, InResource );
+	}
+	else
+	{
+		DynamicNativeTextureMap.Add( ResourceName, InResource );
+	}
+
+}
+
+void FDynamicResourceMap::RemoveResource( FName ResourceName, UObject* ResourceObject )
+{
+	if( ResourceObject )
+	{
+		DynamicResourceObjectMap.Remove( ResourceObject );
+	}
+	else
+	{
+		DynamicNativeTextureMap.Remove( ResourceName );
+	}
+}
+
+void FDynamicResourceMap::Empty()
+{
+	DynamicNativeTextureMap.Empty();
+	DynamicResourceObjectMap.Empty();
+}
+
+void FDynamicResourceMap::ReleaseResources()
+{
+	for (TMap<FName, TSharedPtr<FDynamicTextureResource> >::TIterator It(DynamicNativeTextureMap); It; ++It)
+	{
+		BeginReleaseResource(It.Value()->RHIRefTexture);
+	}
+
+	for (TMap<UObject*, TSharedPtr<FDynamicTextureResource> >::TIterator It(DynamicResourceObjectMap); It; ++It)
+	{
+		BeginReleaseResource(It.Value()->RHIRefTexture);
+	}
+}
+
+void FDynamicResourceMap::AddReferencedObjects(FReferenceCollector& Collector)
+{
+	for(TMap<FName, TSharedPtr<FDynamicTextureResource> >::TIterator It(DynamicNativeTextureMap); It; ++It)
+	{
+		TSharedPtr<FDynamicTextureResource>& Resource = It.Value();
+		if ( Resource.IsValid() && Resource->TextureObject != nullptr )
+		{
+			Collector.AddReferencedObject(Resource->TextureObject);
+		}
+	}
+
+	for(TMap<UObject*, TSharedPtr<FDynamicTextureResource> >::TIterator It(DynamicResourceObjectMap); It; ++It)
+	{
+		TSharedPtr<FDynamicTextureResource>& Resource = It.Value();
+		if (Resource.IsValid() && Resource->TextureObject != nullptr)
+		{
+			Collector.AddReferencedObject(Resource->TextureObject);
+		}
+	}
+}
+
 FSlateRHIResourceManager::FSlateRHIResourceManager()
 {
 	if( GIsEditor )
@@ -233,23 +311,11 @@ FSlateShaderResourceProxy* FSlateRHIResourceManager::GenerateTextureResource( co
 	return NewProxy;
 }
 
-void FSlateRHIResourceManager::AddReferencedObjects( FReferenceCollector& Collector )
-{
-	for( TMap<FName, TSharedPtr<FDynamicTextureResource> >::TIterator It(DynamicTextureMap); It; ++It )
-	{
-		TSharedPtr<FDynamicTextureResource>& Resource = It.Value();
-		if( Resource->TextureObject )
-		{
-			Collector.AddReferencedObject( Resource->TextureObject );
-		}
-	}
-}
 
 FSlateShaderResourceProxy* FSlateRHIResourceManager::GetTexture( const FSlateBrush& InBrush )
 {
 	check( IsThreadSafeForSlateRendering() );
 
-
 	FSlateShaderResourceProxy* Texture = NULL;
 	if( !InBrush.IsDynamicallyLoaded() && !InBrush.HasUObject() )
 	{
@@ -263,10 +329,10 @@ FSlateShaderResourceProxy* FSlateRHIResourceManager::GetTexture( const FSlateBru
 	return Texture;
 }
 
-TSharedPtr<FSlateRHIResourceManager::FDynamicTextureResource> FSlateRHIResourceManager::MakeDynamicTextureResource( FName ResourceName, uint32 Width, uint32 Height, const TArray< uint8 >& Bytes )
+TSharedPtr<FDynamicTextureResource> FSlateRHIResourceManager::MakeDynamicTextureResource( FName ResourceName, uint32 Width, uint32 Height, const TArray< uint8 >& Bytes )
 {
 	// Bail out if we already have this texture loaded
-	TSharedPtr<FDynamicTextureResource> TextureResource = DynamicTextureMap.FindRef( ResourceName );
+	TSharedPtr<FDynamicTextureResource> TextureResource = DynamicResourceMap.GetResource( ResourceName, nullptr );
 	if( TextureResource.IsValid() )
 	{
 		return TextureResource;
@@ -279,35 +345,42 @@ TSharedPtr<FSlateRHIResourceManager::FDynamicTextureResource> FSlateRHIResourceM
 	TextureResource = InitializeDynamicTextureResource( TextureStorage, NULL );
 
 	// Map the new resource so we don't have to load again
-	DynamicTextureMap.Add( ResourceName, TextureResource );
+	DynamicResourceMap.AddResource( ResourceName, nullptr, TextureResource.ToSharedRef() );
 
 	return TextureResource;
 }
 
-TSharedPtr<FSlateRHIResourceManager::FDynamicTextureResource> FSlateRHIResourceManager::MakeDynamicTextureResource(bool bHasUTexture, bool bIsDynamicallyLoaded, FString ResourcePath, FName ResourceName, UTexture2D* InTextureObject)
+TSharedPtr<FDynamicTextureResource> FSlateRHIResourceManager::MakeDynamicTextureResource(bool bHasUTexture, bool bIsDynamicallyLoaded, FString ResourcePath, FName ResourceName, UTexture2D* InTextureObject)
 {
-	// Bail out if we already have this texture loaded
-	TSharedPtr<FDynamicTextureResource> TextureResource = DynamicTextureMap.FindRef( ResourceName );
-	if( TextureResource.IsValid() )
-	{
-		return TextureResource;
-	}
-
+	// Generated texture resource
+	TSharedPtr<FDynamicTextureResource> TextureResource;
+	
 	// Texture object if any
 	UTexture2D* TextureObject = NULL;
 
 	// Data for a loaded disk image
 	FNewTextureInfo Info;
 
+	bool bUsingDeprecatedUTexturePath = false;
+
 	bool bSucceeded = false;
 	if( bHasUTexture || InTextureObject != NULL )
 	{
+		TextureResource = DynamicResourceMap.GetResource( ResourceName, InTextureObject );
+		if( TextureResource.IsValid() )
+		{
+			// Bail out of the resource is already loaded
+			return TextureResource;
+		}
+
 		if( InTextureObject )
 		{
 			TextureObject = InTextureObject;
 		}
 		else
 		{
+			// Deprecated path
+			bUsingDeprecatedUTexturePath = true;
 			// Load the utexture
 			FString Path = ResourceName.ToString();
 			Path = Path.RightChop( FSlateBrush::UTextureIdentifier().Len() );
@@ -318,6 +391,13 @@ TSharedPtr<FSlateRHIResourceManager::FDynamicTextureResource> FSlateRHIResourceM
 	}
 	else if( bIsDynamicallyLoaded )
 	{
+		TextureResource = DynamicResourceMap.GetResource(ResourceName, nullptr );
+		if ( TextureResource.IsValid() )
+		{
+			// Bail out of the resource is already loaded
+			return TextureResource;
+		}
+
 		uint32 Width = 0;
 		uint32 Height = 0;
 		TArray<uint8> RawData;
@@ -328,25 +408,33 @@ TSharedPtr<FSlateRHIResourceManager::FDynamicTextureResource> FSlateRHIResourceM
 		Info.TextureData = MakeShareable( new FSlateTextureData( Width, Height, GPixelFormats[PF_B8G8R8A8].BlockBytes, RawData ) );
 	}
 
+	
 	if( bSucceeded )
 	{
 		TextureResource = InitializeDynamicTextureResource( Info.TextureData, TextureObject );
-
-		// Map the new resource for the UTexture so we don't have to load again
-		DynamicTextureMap.Add( ResourceName, TextureResource );
 	}
 	else
 	{
 		// Add the null texture so we don't continuously try to load it.
-		DynamicTextureMap.Add( ResourceName, FDynamicTextureResource::NullResource );
+		TextureResource = FDynamicTextureResource::NullResource;
 	}
 
+	if( !bUsingDeprecatedUTexturePath )
+	{
+		DynamicResourceMap.AddResource(ResourceName, TextureObject, TextureResource.ToSharedRef());
+	}
+	else
+	{
+		DynamicResourceMap.AddResource(ResourceName, nullptr, TextureResource.ToSharedRef());
+	}
+
+
 	return TextureResource;
 }
 
 
 
-TSharedRef< FSlateRHIResourceManager::FDynamicTextureResource > FSlateRHIResourceManager::InitializeDynamicTextureResource( const FSlateTextureDataPtr& TextureData, UTexture2D* TextureObject )
+TSharedRef<FDynamicTextureResource> FSlateRHIResourceManager::InitializeDynamicTextureResource( const FSlateTextureDataPtr& TextureData, UTexture2D* TextureObject )
 {
 	TSharedPtr<FDynamicTextureResource> TextureResource = NULL;
 
@@ -358,10 +446,10 @@ TSharedRef< FSlateRHIResourceManager::FDynamicTextureResource > FSlateRHIResourc
 	else
 	{
 		// Free list is empty, we have to allocate a new resource
-		TextureResource = MakeShareable( new FDynamicTextureResource( NULL ) );
+		TextureResource = MakeShareable( new FDynamicTextureResource( nullptr ) );
 	}
 
-	checkSlow( !AccessedUTextures.Contains( TextureResource ) );
+	checkSlow( !AccessedUTextures.Contains( TextureObject ) );
 
 	// Init game thread data;
 	TextureResource->TextureObject = TextureObject;
@@ -376,13 +464,14 @@ TSharedRef< FSlateRHIResourceManager::FDynamicTextureResource > FSlateRHIResourc
 	}
 
 	// Init render thread data
-	ENQUEUE_UNIQUE_RENDER_COMMAND_TWOPARAMETER( InitNewSlateDynamicTextureResource,
+	ENQUEUE_UNIQUE_RENDER_COMMAND_THREEPARAMETER( InitNewSlateDynamicTextureResource,
 		FDynamicTextureResource*, TextureResource, TextureResource.Get(),
 		FSlateTextureDataPtr, InNewTextureData, TextureData,
+		bool, bHasResourceObject, TextureObject != nullptr,
 	{
-		if( TextureResource->TextureObject != NULL )
+		if( bHasResourceObject )
 		{
-			TextureResource->RHIRefTexture->SetRHIRef( NULL, 0, 0 );
+			TextureResource->RHIRefTexture->SetRHIRef( nullptr, 0, 0 );
 
 		}
 		else if( InNewTextureData.IsValid() )
@@ -407,7 +496,7 @@ FSlateShaderResourceProxy* FSlateRHIResourceManager::GetDynamicTextureResource(
 	const FName ResourceName = InBrush.GetResourceName();
 	if ( ResourceName.IsValid() && ResourceName != NAME_None )
 	{
-		TSharedPtr<FDynamicTextureResource> TextureResource = DynamicTextureMap.FindRef( ResourceName );
+		TSharedPtr<FDynamicTextureResource> TextureResource = DynamicResourceMap.GetResource( ResourceName, InBrush.GetResourceObject() );
 
 		if( !TextureResource.IsValid() )
 		{
@@ -423,7 +512,8 @@ FSlateShaderResourceProxy* FSlateRHIResourceManager::GetDynamicTextureResource(
 
 		if( TextureResource.IsValid())
 		{
-			if( TextureResource->TextureObject && !AccessedUTextures.Contains( TextureResource ) && TextureResource->TextureObject->Resource)
+			UTexture2D* TextureObject = TextureResource->TextureObject;
+			if( TextureObject && !AccessedUTextures.Contains( TextureObject ) && TextureObject->Resource )
 			{
 				// Set the texture rendering resource that should be used.  The UTexture resource could change at any time so we must do this each frame
 				ENQUEUE_UNIQUE_RENDER_COMMAND_TWOPARAMETER( UpdateSlateUTextureResource,
@@ -435,7 +525,7 @@ FSlateShaderResourceProxy* FSlateRHIResourceManager::GetDynamicTextureResource(
 					InFTexture->LastRenderTime = FApp::GetCurrentTime();
 				});
 
-				AccessedUTextures.Add( TextureResource );
+				AccessedUTextures.Add( TextureObject );
 			}
 
 			return TextureResource->Proxy;
@@ -459,17 +549,18 @@ void FSlateRHIResourceManager::ReleaseDynamicResource( const FSlateBrush& InBrus
 	if( InBrush.HasUObject() || InBrush.IsDynamicallyLoaded() )
 	{
 		FName ResourceName = InBrush.GetResourceName();
-		TSharedPtr<FDynamicTextureResource> TextureResource = DynamicTextureMap.FindRef(ResourceName);
+		TSharedPtr<FDynamicTextureResource> TextureResource = DynamicResourceMap.GetResource( ResourceName, InBrush.GetResourceObject() );
+
 		if( TextureResource.IsValid() )
 		{
 			//remove it from the accessed textures
-			AccessedUTextures.Remove( TextureResource );
+			AccessedUTextures.Remove( TextureResource->TextureObject );
 
 			// Release the rendering resource, its no longer being used
 			BeginReleaseResource( TextureResource->RHIRefTexture );
 
 			//remove it from the texture map
-			DynamicTextureMap.Remove( ResourceName );
+			DynamicResourceMap.RemoveResource( ResourceName, InBrush.GetResourceObject() );
 
 			// Add the resource to the free list so it can be reused
 			DynamicTextureFreeList.Add( TextureResource );
@@ -522,11 +613,7 @@ void FSlateRHIResourceManager::ReleaseResources()
 		BeginReleaseResource( NonAtlasedTextures[ResourceIndex] );
 	}
 
-	
-	for( TMap<FName,TSharedPtr<FDynamicTextureResource> >::TIterator It(DynamicTextureMap); It; ++It )
-	{
-		BeginReleaseResource( It.Value()->RHIRefTexture );
-	}
+	DynamicResourceMap.ReleaseResources();
 
 	for( int32 ResourceIndex = 0; ResourceIndex < DynamicTextureFreeList.Num(); ++ResourceIndex )
 	{
@@ -552,7 +639,7 @@ void FSlateRHIResourceManager::DeleteResources()
 	SET_DWORD_STAT(STAT_SlateNumDynamicTextures, 0);
 
 	AccessedUTextures.Empty();
-	DynamicTextureMap.Empty();
+	DynamicResourceMap.Empty();
 	DynamicTextureFreeList.Empty();
 	TextureAtlases.Empty();
 	NonAtlasedTextures.Empty();
diff --git a/Engine/Source/Runtime/SlateRHIRenderer/Private/SlateRHIResourceManager.h b/Engine/Source/Runtime/SlateRHIRenderer/Private/SlateRHIResourceManager.h
--- a/Engine/Source/Runtime/SlateRHIRenderer/Private/SlateRHIResourceManager.h
+++ b/Engine/Source/Runtime/SlateRHIRenderer/Private/SlateRHIResourceManager.h
@@ -3,26 +3,50 @@
 
 #pragma once
 
+/** Represents a dynamic texture resource for rendering */
+struct FDynamicTextureResource
+{
+	static TSharedPtr<FDynamicTextureResource> NullResource;
 
-/**
- * Stores a mapping of texture names to their RHI texture resource               
- */
-class FSlateRHIResourceManager : public FSlateShaderResourceManager, public FGCObject
+	FDynamicTextureResource(FSlateTexture2DRHIRef* ExistingTexture);
+	~FDynamicTextureResource();
+
+	// @todo Slate This is needed for the deprecated utexture loading path
+	UTexture2D* TextureObject;
+
+	FSlateShaderResourceProxy* Proxy;
+	FSlateTexture2DRHIRef* RHIRefTexture;
+};
+
+
+struct FDynamicResourceMap : public FGCObject
 {
 public:
-	/** Represents a dynamic texture resource for rendering */
-	struct FDynamicTextureResource
-	{
-		class UTexture2D* TextureObject;
-		FSlateShaderResourceProxy* Proxy;
-		FSlateTexture2DRHIRef* RHIRefTexture;
+	TSharedPtr<FDynamicTextureResource> GetResource( FName ResourceName, UObject* ResourceObject ) const;
 
-		static TSharedPtr<FDynamicTextureResource> NullResource;
+	void AddResource( FName ResourceName, UObject* ResourceObject, TSharedRef<FDynamicTextureResource> InResource );
 
-		FDynamicTextureResource( FSlateTexture2DRHIRef* ExistingTexture );
-		~FDynamicTextureResource();
-	};
+	void RemoveResource( FName ResourceName, UObject* ResourceObject );
+
+	void Empty();
+
+	void ReleaseResources();
 
+	/** FGCObject interface */
+	virtual void AddReferencedObjects(FReferenceCollector& Collector) override;
+
+private:
+	TMap<FName, TSharedPtr<FDynamicTextureResource> > DynamicNativeTextureMap;
+
+	TMap<UObject*, TSharedPtr<FDynamicTextureResource> > DynamicResourceObjectMap;
+};
+
+
+/**
+ * Stores a mapping of texture names to their RHI texture resource               
+ */
+class FSlateRHIResourceManager : public FSlateShaderResourceManager
+{
 public:
 	FSlateRHIResourceManager();
 	virtual ~FSlateRHIResourceManager();
@@ -72,7 +96,7 @@ public:
 	 * @param Height				The height of the resource
 	 * @param Bytes					The payload containing the resource
 	 */
-	TSharedPtr<FSlateRHIResourceManager::FDynamicTextureResource> MakeDynamicTextureResource( FName ResourceName, uint32 Width, uint32 Height, const TArray< uint8 >& Bytes );
+	TSharedPtr<FDynamicTextureResource> MakeDynamicTextureResource( FName ResourceName, uint32 Width, uint32 Height, const TArray< uint8 >& Bytes );
 
 	/**
 	 * Returns true if a texture resource with the passed in resource name is available 
@@ -131,7 +155,7 @@ private:
 	/** 
 	 * Prepares a dynamic texture resource for use based on the provided inputs
 	 */
-	TSharedRef< FSlateRHIResourceManager::FDynamicTextureResource > InitializeDynamicTextureResource( const FSlateTextureDataPtr& TextureStorage, UTexture2D* TextureObject );
+	TSharedRef<FDynamicTextureResource> InitializeDynamicTextureResource( const FSlateTextureDataPtr& TextureStorage, UTexture2D* TextureObject );
 
 	/**
 	 * Generates rendering resources for a texture
@@ -148,14 +172,11 @@ private:
 	 */
 	FSlateShaderResourceProxy* GetDynamicTextureResource( const FSlateBrush& InBrush );
 
-	/** FGCObject */
-	void AddReferencedObjects( FReferenceCollector& Collector ) override;
-
 private:
-	/** Map of all active dynamic texture objects being used for brush resources */
-	TMap<FName, TSharedPtr<FDynamicTextureResource> > DynamicTextureMap;
+	/** Map of all active dynamic resources being used by brushes */
+	FDynamicResourceMap DynamicResourceMap;
 	/** Set of dynamic textures that are currently being accessed */
-	TSet< TSharedPtr<FDynamicTextureResource> > AccessedUTextures;
+	TSet< TWeakObjectPtr<UTexture2D> > AccessedUTextures;
 	/** List of old dynamic resources that are free to use as new resources */
 	TArray< TSharedPtr<FDynamicTextureResource> > DynamicTextureFreeList;
 	/** Static Texture atlases which have been created */
