Adding all stubs for material rendering in Slate

diff --git a/Engine/Shaders/SlateElementPixelShader.usf b/Engine/Shaders/SlateElementPixelShader.usf
--- a/Engine/Shaders/SlateElementPixelShader.usf
+++ b/Engine/Shaders/SlateElementPixelShader.usf
@@ -3,6 +3,9 @@
 
 #include "Common.usf"
 
+#if USE_MATERIALS
+#include "Material.usf"
+#endif
 
 /** Shader types (mirrored from ESlateShader::Type in RenderingCommon.h */
 #define ST_Default		0
@@ -23,21 +26,18 @@
 /** Enable to debug clipping issues */
 #define DEBUG_CLIPPING 0
 
-/** Per element params */
-float4 ShaderParams;
-/** Viewport size used for clipping. */
-float2 ViewportSize;
-/** Per element shader type */
-int ShaderType;
-/** Per element draw effects */
-int DrawEffects;
 /** Display gamma */
 float DisplayGamma;
+/** Per element params */
+float4 ShaderParams;
 
+#if !USE_MATERIALS
 /** Texture for the element */
 Texture2D ElementTexture;
 SamplerState ElementTextureSampler;
 
+#endif
+
 struct VertexOut
 {
 	float4 Position				: SV_POSITION;
@@ -47,15 +47,46 @@ struct VertexOut
 	float2 WorldPosition		: TEXCOORD2;
 };
 	
+#if USE_MATERIALS
+float4 GetMaterialColor( VertexOut InVertex, float2 UV )
+{
+	FMaterialPixelParameters Parameters = MakeInitializedMaterialPixelParameters();
+	CalcMaterialParameters(Parameters, true, InVertex.Position);
+#if NUM_MATERIAL_TEXCOORDS
+	for(int CoordinateIndex = 0;CoordinateIndex < NUM_MATERIAL_TEXCOORDS;CoordinateIndex++)
+	{
+		Parameters.TexCoords[CoordinateIndex] = InVertex.TextureCoordinates.xy*InVertex.TextureCoordinates.zw;
+	}
+#endif
+
+	// Clip if masked
+	GetMaterialCoverageAndClipping( Parameters );
+
+	return float4( GetMaterialEmissive( Parameters ), GetMaterialOpacity(Parameters) );
+}
+#endif
+
+#if !USE_MATERIALS
+float4 GetTextureColor( VertexOut InVertex, float2 UV )
+{
+	float4 BaseColor = Texture2DSample(ElementTexture, ElementTextureSampler, UV );
+#if !USE_TEXTURE_ALPHA
+	BaseColor.a = 1.0f;
+#endif
+
+	return BaseColor;
+}
+#endif
 
 /** Gets a color from a texture and supplied vertex color.  Adjusting for tint masking if necessary */
 float4 GetColor( VertexOut InVertex, float2 UV )
 {
 	float4 FinalColor;
 	
-	float4 BaseColor = Texture2DSample(ElementTexture, ElementTextureSampler, UV );
-#if !USE_TEXTURE_ALPHA
-	BaseColor.a = 1.0f;
+#if USE_MATERIALS
+	float4 BaseColor = GetMaterialColor( InVertex, UV );
+#else
+	float4 BaseColor = GetTextureColor( InVertex, UV );
 #endif
 
 	FinalColor = BaseColor*InVertex.Color;
@@ -173,8 +204,8 @@ float4 GetLineSegmentElementColor( VertexOut InVertex )
 	}
 
 	return Color;
-}
-#endif  
+} 
+#endif
 
 float4 ApplyColorVisionDeficiencyFilter(float4 Color)
 {
diff --git a/Engine/Source/Editor/DetailCustomizations/Private/SlateBrushCustomization.cpp b/Engine/Source/Editor/DetailCustomizations/Private/SlateBrushCustomization.cpp
--- a/Engine/Source/Editor/DetailCustomizations/Private/SlateBrushCustomization.cpp
+++ b/Engine/Source/Editor/DetailCustomizations/Private/SlateBrushCustomization.cpp
@@ -974,7 +974,7 @@ void FSlateBrushStructCustomization::CustomizeChildren( TSharedRef<IPropertyHand
 		SNew( SObjectPropertyEntryBox )
 		.PropertyHandle( ResourceObjectProperty )
 		.ThumbnailPool( StructCustomizationUtils.GetThumbnailPool() )
-		.AllowedClass( UTexture2D::StaticClass() )
+		.OnShouldFilterAsset( this, &FSlateBrushStructCustomization::OnFilterAssetPicker )
 	];
 
 	StructBuilder.AddChildProperty( ImageSizeProperty.ToSharedRef() );
@@ -1028,6 +1028,12 @@ void FSlateBrushStructCustomization::CustomizeChildren( TSharedRef<IPropertyHand
 	}
 }
 
+bool FSlateBrushStructCustomization::OnFilterAssetPicker( const FAssetData& InAssetData ) const
+{
+	UClass* Class = InAssetData.GetClass();
+	return !Class->IsChildOf( UTexture2D::StaticClass() );//&& !Class->IsChildOf( UMaterialInterface::StaticClass() );
+}
+
 EVisibility FSlateBrushStructCustomization::GetTilingPropertyVisibility() const
 {
 	uint8 DrawAsType;
diff --git a/Engine/Source/Editor/DetailCustomizations/Private/SlateBrushCustomization.h b/Engine/Source/Editor/DetailCustomizations/Private/SlateBrushCustomization.h
--- a/Engine/Source/Editor/DetailCustomizations/Private/SlateBrushCustomization.h
+++ b/Engine/Source/Editor/DetailCustomizations/Private/SlateBrushCustomization.h
@@ -23,6 +23,9 @@ private:
 	 *  Get the Slate Brush margin property row visibility
 	 */
 	EVisibility GetMarginPropertyVisibility() const;
+	
+	/** Called when the asset picker needs to be filtered */
+	bool OnFilterAssetPicker( const FAssetData& InAssetData ) const;
 
 	/** Slate Brush DrawAs property */
 	TSharedPtr<IPropertyHandle> DrawAsProperty;
diff --git a/Engine/Source/Runtime/Engine/Classes/Materials/Material.h b/Engine/Source/Runtime/Engine/Classes/Materials/Material.h
--- a/Engine/Source/Runtime/Engine/Classes/Materials/Material.h
+++ b/Engine/Source/Runtime/Engine/Classes/Materials/Material.h
@@ -575,6 +575,13 @@ public:
 	UPROPERTY(EditAnywhere, BlueprintReadOnly, Category=Usage)
 	uint32 bUsedWithClothing:1;
 
+	/** 
+	 * Indicates that the material and its instances can be use with Slate UI and UMG
+	 * This will result in the shaders required to support UI materials being compiled which will increase shader compile time and memory usage.
+	 */
+	UPROPERTY(EditAnywhere, BlueprintReadOnly, Category=Usage)
+	uint32 bUsedWithUI:1;
+
 	/* Forces the material to be completely rough. Saves a number of instructions and one sampler. */
 	UPROPERTY(EditAnywhere, BlueprintReadOnly, Category=Mobile)
 	uint32 bFullyRough:1;
diff --git a/Engine/Source/Runtime/Engine/Private/Materials/MaterialShared.cpp b/Engine/Source/Runtime/Engine/Private/Materials/MaterialShared.cpp
--- a/Engine/Source/Runtime/Engine/Private/Materials/MaterialShared.cpp
+++ b/Engine/Source/Runtime/Engine/Private/Materials/MaterialShared.cpp
@@ -746,6 +746,11 @@ bool FMaterialResource::IsUsedWithAPEXCloth() const
 	return Material->bUsedWithClothing;
 }
 
+bool FMaterialResource::IsUsedWithUI() const
+{
+	return Material->bUsedWithUI;
+}
+
 EMaterialTessellationMode FMaterialResource::GetTessellationMode() const 
 { 
 	return (EMaterialTessellationMode)Material->D3D11TessellationMode; 
diff --git a/Engine/Source/Runtime/Engine/Public/MaterialShared.h b/Engine/Source/Runtime/Engine/Public/MaterialShared.h
--- a/Engine/Source/Runtime/Engine/Public/MaterialShared.h
+++ b/Engine/Source/Runtime/Engine/Public/MaterialShared.h
@@ -900,6 +900,7 @@ public:
 	virtual bool IsUsedWithSplineMeshes() const { return false; }
 	virtual bool IsUsedWithInstancedStaticMeshes() const { return false; }
 	virtual bool IsUsedWithAPEXCloth() const { return false; }
+	virtual bool IsUsedWithUI() const { return false; }
 	ENGINE_API virtual enum EMaterialTessellationMode GetTessellationMode() const;
 	virtual bool IsCrackFreeDisplacementEnabled() const { return false; }
 	virtual bool IsAdaptiveTessellationEnabled() const { return false; }
@@ -1437,6 +1438,7 @@ public:
 	ENGINE_API virtual bool IsUsedWithSplineMeshes() const;
 	ENGINE_API virtual bool IsUsedWithInstancedStaticMeshes() const;
 	ENGINE_API virtual bool IsUsedWithAPEXCloth() const;
+	ENGINE_API virtual bool IsUsedWithUI() const;
 	ENGINE_API virtual enum EMaterialTessellationMode GetTessellationMode() const;
 	ENGINE_API virtual bool IsCrackFreeDisplacementEnabled() const;
 	ENGINE_API virtual bool IsAdaptiveTessellationEnabled() const;
diff --git a/Engine/Source/Runtime/Engine/Public/Slate/SlateTextures.h b/Engine/Source/Runtime/Engine/Public/Slate/SlateTextures.h
--- a/Engine/Source/Runtime/Engine/Public/Slate/SlateTextures.h
+++ b/Engine/Source/Runtime/Engine/Public/Slate/SlateTextures.h
@@ -117,6 +117,28 @@ private:
 	uint32 Height;
 };
 
+
+class FSlateMaterial : public FSlateShaderResource
+{
+public:
+	FSlateMaterial( FMaterialRenderProxy* InProxy,  const FVector2D& InImageSize )
+		: RenderProxy( InProxy )
+		, Width( FMath::RoundToInt( InImageSize.X ) )
+		, Height( FMath::RoundToInt( InImageSize.Y ) )
+	{}
+
+	virtual uint32 GetWidth() const { return Width; } 
+	virtual uint32 GetHeight() const { return Height; } 
+
+	virtual void SetResource( FMaterialRenderProxy* Proxy ) { RenderProxy = Proxy; }
+	virtual FMaterialRenderProxy* GetRenderProxy() const { return RenderProxy; }
+	virtual ESlateShaderResource::Type GetType() const override { return ESlateShaderResource::Material; }
+private:
+	FMaterialRenderProxy* RenderProxy;
+	uint32 Width;
+	uint32 Height;
+};
+
 class FSlateTextureRenderTarget2DResource : public FTextureRenderTargetResource
 {
 public:
diff --git a/Engine/Source/Runtime/Slate/Public/Slate.h b/Engine/Source/Runtime/Slate/Public/Slate.h
--- a/Engine/Source/Runtime/Slate/Public/Slate.h
+++ b/Engine/Source/Runtime/Slate/Public/Slate.h
@@ -17,6 +17,7 @@
 #define BEGIN_SLATE_FUNCTION_BUILD_OPTIMIZATION PRAGMA_DISABLE_OPTIMIZATION
 #define END_SLATE_FUNCTION_BUILD_OPTIMIZATION   PRAGMA_ENABLE_OPTIMIZATION
 
+#define ALLOW_SLATE_MATERIALS 0
 
 /** Stat tracking for Slate */
 DECLARE_CYCLE_STAT_EXTERN( TEXT("Total Slate Tick Time"), STAT_SlateTickTime, STATGROUP_Slate , );
diff --git a/Engine/Source/Runtime/Slate/Public/SlateMaterialBrush.h b/Engine/Source/Runtime/Slate/Public/SlateMaterialBrush.h
new file mode 100644
--- /dev/null
+++ b/Engine/Source/Runtime/Slate/Public/SlateMaterialBrush.h
@@ -0,0 +1,27 @@
+// Copyright 1998-2013 Epic Games, Inc. All Rights Reserved.
+
+#pragma once
+
+#include "Slate.h"
+
+#if ALLOW_SLATE_MATERIALS
+
+struct FSlateMaterialBrush : public FSlateBrush
+{
+	FSlateMaterialBrush( class UMaterialInterface& InMaterial, const FVector2D& InImageSize )
+		: FSlateBrush( ESlateBrushDrawType::Image, FName(TEXT("None")), FMargin(0), ESlateBrushTileType::NoTile, ESlateBrushImageType::FullColor, InImageSize, FLinearColor::White, &InMaterial )
+	{
+		ResourceName = FName( *InMaterial.GetFullName() );
+	}
+
+
+	virtual ~FSlateMaterialBrush()
+	{
+		if(FSlateApplication::IsInitialized() && FSlateApplication::Get().GetRenderer().IsValid())
+		{
+			FSlateApplication::Get().GetRenderer()->ReleaseDynamicResource( *this );
+		}
+	}
+}; 
+
+#endif
\ No newline at end of file
diff --git a/Engine/Source/Runtime/SlateCore/Private/Rendering/ElementBatcher.cpp b/Engine/Source/Runtime/SlateCore/Private/Rendering/ElementBatcher.cpp
--- a/Engine/Source/Runtime/SlateCore/Private/Rendering/ElementBatcher.cpp
+++ b/Engine/Source/Runtime/SlateCore/Private/Rendering/ElementBatcher.cpp
@@ -51,7 +51,9 @@ static FVector2D RotatePoint( FVector2D InPoint, const FVector2D& AboutPoint, fl
 
 
 FSlateElementBatcher::FSlateElementBatcher( TSharedRef<FSlateRenderingPolicy> InRenderingPolicy )
-	: RenderingPolicy( InRenderingPolicy )
+	: ResourceManager( *InRenderingPolicy->GetResourceManager() )
+	, FontCache( *InRenderingPolicy->GetFontCache() )
+	, PixelCenterOffset( InRenderingPolicy->GetPixelCenterOffset() )
 {
 
 }
@@ -180,10 +182,10 @@ void FSlateElementBatcher::AddBoxElement( const FVector2D& InPosition, const FVe
 		FVector2D StartUV = FVector2D(0.0f,0.0f);
 		FVector2D EndUV = FVector2D(1.0f,1.0f);
 		FVector2D SizeUV;
-		const float PixelCenterOffset = RenderingPolicy->GetPixelCenterOffset();
+
 		FVector2D HalfTexel;
 
-		FSlateShaderResourceProxy* ResourceProxy = RenderingPolicy->GetTextureResource( *InPayload.BrushResource );
+		FSlateShaderResourceProxy* ResourceProxy = ResourceManager.GetShaderResource( *InPayload.BrushResource );
 		FSlateShaderResource* Resource = nullptr;
 		if( ResourceProxy )
 		{
@@ -202,7 +204,7 @@ void FSlateElementBatcher::AddBoxElement( const FVector2D& InPosition, const FVe
 		}
 		else
 		{
-			// nullptr texture
+			// no texture
 			SizeUV = FVector2D(1.0f,1.0f);
 			HalfTexel = FVector2D( PixelCenterOffset, PixelCenterOffset );
 		}
@@ -425,8 +427,6 @@ void FSlateElementBatcher::AddTextElement( const FVector2D& Position, const FVec
 		return;
 	}
 
-	FSlateFontCache& FontCache = *RenderingPolicy->GetFontCache().Get();
-
 	const float FontScale = InScale;
 
 	FCharacterList& CharacterList = FontCache.GetCharacterList( InPayload.FontInfo, FontScale );
@@ -703,7 +703,7 @@ void FSlateElementBatcher::AddSplineElement( const FVector2D& Position, float Sc
 	const float HalfThickness = LineThickness * .5f + Radius;
 
 	// Currently splines are not atlased because they are tiled.  So we just assume the texture proxy holds the actual texture
-	FSlateShaderResourceProxy* ResourceProxy = RenderingPolicy->GetTextureResource( SplineFilterTable );
+	FSlateShaderResourceProxy* ResourceProxy = ResourceManager.GetShaderResource( SplineFilterTable );
 	check(ResourceProxy && ResourceProxy->Resource);
 
 	// Find a batch for the element
@@ -822,7 +822,7 @@ void FSlateElementBatcher::AddLineElement( const FVector2D& Position, const FVec
 	if( InPayload.bAntialias )
 	{
 		// Currently splines are not atlased because they are tiled.  So we just assume the texture proxy holds the actual texture
-		FSlateShaderResourceProxy* ResourceProxy = RenderingPolicy->GetTextureResource( SplineFilterTable );
+		FSlateShaderResourceProxy* ResourceProxy = ResourceManager.GetShaderResource( SplineFilterTable );
 		check(ResourceProxy && ResourceProxy->Resource);
 
 		// The radius to use when checking the distance of pixels to the actual line.  Arbitrary value based on what looks the best
@@ -995,7 +995,9 @@ void FSlateElementBatcher::AddViewportElement( const FVector2D& InPosition, cons
 
 	TSharedPtr<const ISlateViewport> ViewportPin = InPayload.Viewport.Pin();
 
-	FSlateElementBatch& ElementBatch = FindBatchForElement( Layer, FShaderParams(), RenderingPolicy->GetViewportResource( ViewportPin.Get() ), ESlateDrawPrimitive::TriangleList, ESlateShader::Default, InDrawEffects, DrawFlags );
+	FSlateShaderResource* ViewportResource = ViewportPin.IsValid() ? ViewportPin->GetViewportRenderTargetTexture() : nullptr;
+
+	FSlateElementBatch& ElementBatch = FindBatchForElement( Layer, FShaderParams(), ViewportResource, ESlateDrawPrimitive::TriangleList, ESlateShader::Default, InDrawEffects, DrawFlags );
 	TArray<FSlateVertex>& BatchVertices = BatchVertexArrays[ElementBatch.VertexArrayIndex];
 	TArray<SlateIndex>& BatchIndices = BatchIndexArrays[ElementBatch.IndexArrayIndex];
 
@@ -1057,7 +1059,7 @@ void FSlateElementBatcher::AddBorderElement( const FVector2D& InPosition, const
 	uint32 TextureHeight = 1;
 
 	// Currently borders are not atlased because they are tiled.  So we just assume the texture proxy holds the actual texture
-	FSlateShaderResourceProxy* ResourceProxy = RenderingPolicy->GetTextureResource( *InPayload.BrushResource );
+	FSlateShaderResourceProxy* ResourceProxy = ResourceManager.GetShaderResource( *InPayload.BrushResource );
 	FSlateShaderResource* Resource = ResourceProxy ? ResourceProxy->Resource : nullptr;
 	if( Resource )
 	{
@@ -1066,7 +1068,6 @@ void FSlateElementBatcher::AddBorderElement( const FVector2D& InPosition, const
 	}
  
 	// Texel offset
-	const float PixelCenterOffset = RenderingPolicy->GetPixelCenterOffset();
 	const FVector2D HalfTexel( PixelCenterOffset/TextureWidth, PixelCenterOffset/TextureHeight );
 
 	const FVector2D StartUV = HalfTexel;
@@ -1339,7 +1340,7 @@ FSlateElementBatch& FSlateElementBatcher::FindBatchForElement(
 	return *ElementBatch;
 }
 
-void FSlateElementBatcher::AddBasicVertices( TArray<FSlateVertex>& OutVertices, FSlateElementBatch& ElementBatch, const TArray<FSlateVertex>& VertexBatch )
+void FSlateElementBatcher::AddVertices( TArray<FSlateVertex>& OutVertices, FSlateElementBatch& ElementBatch, const TArray<FSlateVertex>& VertexBatch )
 {
 	uint32 FirstIndex = OutVertices.Num();
 	OutVertices.AddUninitialized( VertexBatch.Num() );
@@ -1426,8 +1427,7 @@ void FSlateElementBatcher::FillBatchBuffers( FSlateWindowElementList& WindowElem
 
 				if( BatchVertices.Num() > 0 && BatchIndices.Num() > 0  )
 				{
-					
-					AddBasicVertices( OutBatchedVertices, ElementBatch, BatchVertices );
+					AddVertices( OutBatchedVertices, ElementBatch, BatchVertices );
 					AddIndices( OutBatchedIndices, ElementBatch, BatchIndices );
 
 					OutRenderBatches.Add( FSlateRenderBatch( ElementBatch ) );
diff --git a/Engine/Source/Runtime/SlateCore/Public/Rendering/ElementBatcher.h b/Engine/Source/Runtime/SlateCore/Public/Rendering/ElementBatcher.h
--- a/Engine/Source/Runtime/SlateCore/Public/Rendering/ElementBatcher.h
+++ b/Engine/Source/Runtime/SlateCore/Public/Rendering/ElementBatcher.h
@@ -179,7 +179,7 @@ private:
 											 ESlateDrawEffect::Type DrawEffects, 
 											 ESlateBatchDrawFlag::Type DrawFlags = ESlateBatchDrawFlag::None);
 
-	void AddBasicVertices( TArray<FSlateVertex>& OutVertices, FSlateElementBatch& ElementBatch, const TArray<FSlateVertex>& VertexBatch );
+	void AddVertices( TArray<FSlateVertex>& OutVertices, FSlateElementBatch& ElementBatch, const TArray<FSlateVertex>& VertexBatch );
 
 	void AddIndices( TArray<SlateIndex>& OutIndices, FSlateElementBatch& ElementBatch, const TArray<SlateIndex>& IndexBatch );
 
@@ -202,7 +202,14 @@ private:
 	// Array of vertex lists for batching indices. We use this method for quickly resetting the arrays without deleting memory.
 	TArray<TArray<SlateIndex>> BatchIndexArrays;
 
-	TSharedRef<FSlateRenderingPolicy> RenderingPolicy;
+	/** Resource manager for accessing shader resources */
+	FSlateShaderResourceManager& ResourceManager;
+
+	/** Font cache used to layout text */
+	FSlateFontCache& FontCache;
+
+	/** Offset to use when supporting 1:1 texture to pixel snapping */
+	const float PixelCenterOffset;
 
 	// true if any element in the batch requires vsync.
 	bool bRequiresVsync;
diff --git a/Engine/Source/Runtime/SlateCore/Public/Rendering/RenderingPolicy.h b/Engine/Source/Runtime/SlateCore/Public/Rendering/RenderingPolicy.h
--- a/Engine/Source/Runtime/SlateCore/Public/Rendering/RenderingPolicy.h
+++ b/Engine/Source/Runtime/SlateCore/Public/Rendering/RenderingPolicy.h
@@ -33,15 +33,11 @@ public:
 	 */
 	virtual ~FSlateRenderingPolicy( ) { }
 
-public:
-
 	virtual void UpdateBuffers( const FSlateWindowElementList& WindowElementList ) = 0;
 
-	virtual FSlateShaderResource* GetViewportResource( const ISlateViewport* InViewportInterface ) = 0;
-
-	virtual FSlateShaderResourceProxy* GetTextureResource( const FSlateBrush& Brush ) = 0;
+	virtual TSharedPtr<class FSlateFontCache> GetFontCache( ) = 0;
 
-	virtual TSharedPtr<class FSlateFontCache>& GetFontCache( ) = 0;
+	virtual TSharedRef<class FSlateShaderResourceManager> GetResourceManager() = 0;
 
 	float GetPixelCenterOffset( ) const
 	{
@@ -49,6 +45,8 @@ public:
 	}
 
 private:
-
+	FSlateRenderingPolicy(const FSlateRenderingPolicy&);
+	FSlateRenderingPolicy& operator=(const FSlateRenderingPolicy&);
+private:
 	float PixelCenterOffset;
 };
diff --git a/Engine/Source/Runtime/SlateCore/Public/Rendering/SlateRenderer.h b/Engine/Source/Runtime/SlateCore/Public/Rendering/SlateRenderer.h
--- a/Engine/Source/Runtime/SlateCore/Public/Rendering/SlateRenderer.h
+++ b/Engine/Source/Runtime/SlateCore/Public/Rendering/SlateRenderer.h
@@ -16,7 +16,7 @@ class SWindow;
 class FSlateRenderer
 {
 public:
-
+	FSlateRenderer(){}
 	/**
 	 * Virtual destructor.
 	 */
@@ -174,6 +174,10 @@ public:
 	 */
 	virtual void PrepareToTakeScreenshot(const FIntRect& Rect, TArray<FColor>* OutColorData) {}
 
+private:
+	// Non-copyable
+	FSlateRenderer(const FSlateRenderer&);
+	FSlateRenderer& operator=(const FSlateRenderer&);
 protected:
 
 	TSharedPtr<class FSlateFontCache> FontCache;
diff --git a/Engine/Source/Runtime/SlateCore/Public/Textures/TextureManager.h b/Engine/Source/Runtime/SlateCore/Public/Textures/TextureManager.h
--- a/Engine/Source/Runtime/SlateCore/Public/Textures/TextureManager.h
+++ b/Engine/Source/Runtime/SlateCore/Public/Textures/TextureManager.h
@@ -46,7 +46,7 @@ public:
 	/** 
 	 * Returns a texture associated with the passed in name.  Should return nullptr if not found 
 	 */
-	virtual FSlateShaderResourceProxy* GetTexture( const FSlateBrush& InBrush ) = 0;
+	virtual FSlateShaderResourceProxy* GetShaderResource( const FSlateBrush& InBrush ) = 0;
 
 protected:
 
@@ -60,7 +60,7 @@ protected:
 		ResourceMap.Empty();
 	}
 
-	FString GetResourcePath( const FSlateBrush& InBrush )
+	FString GetResourcePath( const FSlateBrush& InBrush ) const
 	{
 		// assume the brush name contains the whole path
 		return InBrush.GetResourceName().ToString();
@@ -68,4 +68,10 @@ protected:
 
 	/** Mapping of names to texture pointers */
 	TMap<FName,FSlateShaderResourceProxy*> ResourceMap;
+
+private:
+	// Non-copyable
+	FSlateShaderResourceManager(const FSlateShaderResourceManager&);
+	FSlateShaderResourceManager& operator=(const FSlateShaderResourceManager&);
+
 };
diff --git a/Engine/Source/Runtime/SlateRHIRenderer/Private/SlateRHIRenderer.cpp b/Engine/Source/Runtime/SlateRHIRenderer/Private/SlateRHIRenderer.cpp
--- a/Engine/Source/Runtime/SlateRHIRenderer/Private/SlateRHIRenderer.cpp
+++ b/Engine/Source/Runtime/SlateRHIRenderer/Private/SlateRHIRenderer.cpp
@@ -14,6 +14,19 @@ namespace CrashTrackerConstants
 }
 
 
+class FSlateBackBuffer : public FRenderTarget
+{
+public:
+	FSlateBackBuffer(FTexture2DRHIRef& InRenderTargetTexture, FIntPoint InSizeXY)
+		: SizeXY(InSizeXY)
+	{
+		RenderTargetTextureRHI = InRenderTargetTexture;
+	}
+	virtual FIntPoint GetSizeXY() const override { return SizeXY; }
+private:
+	FIntPoint SizeXY;
+};
+
 // Defines the maximum size that a slate viewport will create
 #define MAX_VIEWPORT_SIZE 16384
 
@@ -446,7 +459,7 @@ void FSlateRHIRenderer::DrawWindow_RenderThread( const FViewportInfo& ViewportIn
 
 		if( WindowElementList.GetRenderBatches().Num() > 0 )
 		{
-			FSlateRenderTarget BackBufferTarget( BackBuffer, FIntPoint( ViewportInfo.Width, ViewportInfo.Height ) );
+			FSlateBackBuffer BackBufferTarget( BackBuffer, FIntPoint( ViewportInfo.Width, ViewportInfo.Height ) );
 
 			RenderingPolicy->DrawElements
 			(
@@ -826,7 +839,7 @@ void FSlateRHIRenderer::CopyWindowsToDrawBuffer(const TArray<FString>& KeypressB
 		if( Context.SlateElementList->GetRenderBatches().Num() > 0 )
 		{
 			FTexture2DRHIRef UnusedTargetTexture;
-			FSlateRenderTarget UnusedTarget( UnusedTargetTexture, FIntPoint::ZeroValue );
+			FSlateBackBuffer UnusedTarget( UnusedTargetTexture, FIntPoint::ZeroValue );
 
 			Context.RenderPolicy->DrawElements(Context.ViewportSize, UnusedTarget, CreateProjectionMatrix(Context.ViewportSize.X, Context.ViewportSize.Y), Context.SlateElementList->GetRenderBatches());
 		}
diff --git a/Engine/Source/Runtime/SlateRHIRenderer/Private/SlateRHIRenderingPolicy.cpp b/Engine/Source/Runtime/SlateRHIRenderer/Private/SlateRHIRenderingPolicy.cpp
--- a/Engine/Source/Runtime/SlateRHIRenderer/Private/SlateRHIRenderingPolicy.cpp
+++ b/Engine/Source/Runtime/SlateRHIRenderer/Private/SlateRHIRenderingPolicy.cpp
@@ -2,8 +2,12 @@
 
 #include "SlateRHIRendererPrivatePCH.h"
 #include "RenderingPolicy.h"
-#include "RHIStaticStates.h"
-
+#include "SlateRHIRenderingPolicy.h"
+#include "SlateShaders.h"
+#include "SlateRHIResourceManager.h"
+#include "TileRendering.h"
+#include "PreviewScene.h"
+#include "EngineModule.h"
 
 static EPrimitiveType GetRHIPrimitiveType( ESlateDrawPrimitive::Type SlateType )
 {
@@ -18,91 +22,6 @@ static EPrimitiveType GetRHIPrimitiveType( ESlateDrawPrimitive::Type SlateType )
 
 };
 
-FSlateElementVertexBuffer::FSlateElementVertexBuffer()
-	: BufferSize(0)
-	, BufferUsageSize(0)
-{
-
-}
-
-FSlateElementVertexBuffer::~FSlateElementVertexBuffer()
-{
-
-}
-
-/** Initializes the vertex buffers RHI resource. */
-void FSlateElementVertexBuffer::InitDynamicRHI()
-{
-	if( !IsValidRef(VertexBufferRHI) )
-	{
-		if( BufferSize == 0 )
-		{
-			BufferSize = 200 * sizeof(FSlateVertex);
-		}
-
-		FRHIResourceCreateInfo CreateInfo;
-		VertexBufferRHI = RHICreateVertexBuffer( BufferSize, BUF_Dynamic, CreateInfo);
-
-		// Ensure the vertex buffer could be created
-		check(IsValidRef(VertexBufferRHI));
-	}
-}
-
-/** Releases the vertex buffers RHI resource. */
-void FSlateElementVertexBuffer::ReleaseDynamicRHI()
-{
-	VertexBufferRHI.SafeRelease();
-	BufferSize = 0;
-}
-
-void FSlateElementVertexBuffer::FillBuffer( const TArray<FSlateVertex>& InVertices, bool bShrinkToFit )
-{
-	check( IsInRenderingThread() );
-
-	if( InVertices.Num() )
-	{
-		uint32 NumVertices = InVertices.Num();
-
-#if !SLATE_USE_32BIT_INDICES
-		// make sure our index buffer can handle this
-		checkf(NumVertices < 0xFFFF, TEXT("Slate vertex buffer is too large (%d) to work with uint16 indices"), NumVertices);
-#endif
-
-		uint32 RequiredBufferSize = NumVertices*sizeof(FSlateVertex);
-
-		// resize if needed
-		if( RequiredBufferSize > GetBufferSize() || bShrinkToFit )
-		{
-			// Use array resize techniques for the vertex buffer
-			ResizeBuffer( InVertices.GetAllocatedSize() );
-		}
-
-		BufferUsageSize += RequiredBufferSize;
-
-		void* VerticesPtr = RHILockVertexBuffer( VertexBufferRHI, 0, RequiredBufferSize, RLM_WriteOnly );
-
-		FMemory::Memcpy( VerticesPtr, InVertices.GetData(), RequiredBufferSize );
-
-		RHIUnlockVertexBuffer( VertexBufferRHI );
-	}
-
-}
-
-void FSlateElementVertexBuffer::ResizeBuffer( uint32 NewSizeBytes )
-{
-	check( IsInRenderingThread() );
-
-	if( NewSizeBytes != 0 && NewSizeBytes != BufferSize)
-	{
-		VertexBufferRHI.SafeRelease();
-		FRHIResourceCreateInfo CreateInfo;
-		VertexBufferRHI = RHICreateVertexBuffer( NewSizeBytes, BUF_Dynamic, CreateInfo);
-		check(IsValidRef(VertexBufferRHI));
-
-		BufferSize = NewSizeBytes;
-	}
-
-}
 
 FSlateElementIndexBuffer::FSlateElementIndexBuffer()
 	: BufferSize(0)	 
@@ -184,9 +103,9 @@ void FSlateElementIndexBuffer::ReleaseDynamicRHI()
 FGlobalBoundShaderState FSlateRHIRenderingPolicy::NormalShaderStates[4][2 /* UseTextureAlpha */];
 FGlobalBoundShaderState FSlateRHIRenderingPolicy::DisabledShaderStates[4][2 /* UseTextureAlpha */];
 
-FSlateRHIRenderingPolicy::FSlateRHIRenderingPolicy( TSharedPtr<FSlateFontCache> InFontCache, TSharedRef<FSlateRHIResourceManager> InTextureManager )
+FSlateRHIRenderingPolicy::FSlateRHIRenderingPolicy( TSharedPtr<FSlateFontCache> InFontCache, TSharedRef<FSlateRHIResourceManager> InResourceManager )
 	: FSlateRenderingPolicy( GPixelCenterOffset )
-	, TextureManager( InTextureManager )
+	, ResourceManager( InResourceManager )
 	, FontCache( InFontCache )
 	, CurrentBufferIndex(0)
 	, bShouldShrinkResources(false)
@@ -269,17 +188,38 @@ void FSlateRHIRenderingPolicy::UpdateBuffers( const FSlateWindowElementList& Win
 	// Should only be called by the rendering thread
 	check(IsInRenderingThread());
 
-	const TArray<FSlateVertex>& Vertices = WindowElementList.GetBatchedVertices();
-	const TArray<SlateIndex>& Indices = WindowElementList.GetBatchedIndices();
+	{
+		const TArray<FSlateVertex>& Vertices = WindowElementList.GetBatchedVertices();
+		const TArray<SlateIndex>& Indices = WindowElementList.GetBatchedIndices();
 
-	FSlateElementVertexBuffer& VertexBuffer = VertexBuffers[CurrentBufferIndex];
-	FSlateElementIndexBuffer& IndexBuffer = IndexBuffers[CurrentBufferIndex];
+		TSlateElementVertexBuffer<FSlateVertex>& VertexBuffer = VertexBuffers[CurrentBufferIndex];
+		FSlateElementIndexBuffer& IndexBuffer = IndexBuffers[CurrentBufferIndex];
+
+		VertexBuffer.FillBuffer( Vertices, bShouldShrinkResources );
+		IndexBuffer.FillBuffer( Indices, bShouldShrinkResources );
+	}
+}
+
+static FSceneView& CreateSceneView( FSceneViewFamilyContext& ViewFamilyContext, FSlateBackBuffer& BackBuffer, const FMatrix& ViewProjectionMatrix )
+{
+	FIntRect ViewRect(FIntPoint(0, 0), BackBuffer.GetSizeXY());
 
-	VertexBuffer.FillBuffer( Vertices, bShouldShrinkResources );
-	IndexBuffer.FillBuffer( Indices, bShouldShrinkResources );
+	// make a temporary view
+	FSceneViewInitOptions ViewInitOptions;
+	ViewInitOptions.ViewFamily = &ViewFamilyContext;
+	ViewInitOptions.SetViewRectangle(ViewRect);
+	ViewInitOptions.ViewMatrix = FMatrix::Identity;
+	ViewInitOptions.ProjectionMatrix = ViewProjectionMatrix;
+	ViewInitOptions.BackgroundColor = FLinearColor::Black;
+	ViewInitOptions.OverlayColor = FLinearColor::White;
+
+	FSceneView* View = new FSceneView( ViewInitOptions );
+	ViewFamilyContext.Views.Add( View );
+
+	return *View;
 }
 
-void FSlateRHIRenderingPolicy::DrawElements( const FIntPoint& InViewportSize, FSlateRenderTarget& BackBuffer, const FMatrix& ViewProjectionMatrix, const TArray<FSlateRenderBatch>& RenderBatches )
+void FSlateRHIRenderingPolicy::DrawElements( const FIntPoint& InViewportSize, FSlateBackBuffer& BackBuffer, const FMatrix& ViewProjectionMatrix, const TArray<FSlateRenderBatch>& RenderBatches )
 {
 	SCOPE_CYCLE_COUNTER( STAT_SlateDrawTime );
 
@@ -289,47 +229,58 @@ void FSlateRHIRenderingPolicy::DrawElements( const FIntPoint& InViewportSize, FS
 	//@todo-rco: RHIPacketList
 	FRHICommandList& RHICmdList = FRHICommandList::GetNullRef();
 
-	const float DisplayGamma = GEngine ? GEngine->GetDisplayGamma() : 2.2f;
-
-	FSlateElementVertexBuffer& VertexBuffer = VertexBuffers[CurrentBufferIndex];
+	TSlateElementVertexBuffer<FSlateVertex>& VertexBuffer = VertexBuffers[CurrentBufferIndex];
 	FSlateElementIndexBuffer& IndexBuffer = IndexBuffers[CurrentBufferIndex];
 
-	// Set the vertex buffer stream
-	RHICmdList.SetStreamSource(0, VertexBuffer.VertexBufferRHI, sizeof(FSlateVertex), 0);
+	float TimeSeconds = FPlatformTime::Seconds() - GStartTime;
+	float RealTimeSeconds = FPlatformTime::Seconds() - GStartTime;
+	float DeltaTimeSeconds = FApp::GetDeltaTime();
 
-	TShaderMapRef<FSlateElementVS> VertexShader(GetGlobalShaderMap());
+	static const FEngineShowFlags DefaultShowFlags(ESFIM_Game);
 
-	FSamplerStateRHIRef BilinearClamp = TStaticSamplerState<SF_Bilinear,AM_Clamp,AM_Clamp,AM_Clamp>::GetRHI();
+	const float DisplayGamma = GEngine ? GEngine->GetDisplayGamma() : 2.2f;
+
+#if ALLOW_SLATE_MATERIALS 
+	FSceneViewFamilyContext SceneViewContext
+	(
+		FSceneViewFamily::ConstructionValues
+		(
+			&BackBuffer,
+			NULL,
+			DefaultShowFlags
+		)
+		.SetWorldTimes( TimeSeconds, RealTimeSeconds, DeltaTimeSeconds )
+		.SetGammaCorrection( DisplayGamma )
+	);
+
+	FSceneView* SceneView = NULL;
+#endif
+
+	TShaderMapRef<FSlateElementVS> VertexShader(GetGlobalShaderMap());
 
 	// Disabled stencil test state
 	FDepthStencilStateRHIRef DSOff = TStaticDepthStencilState<false,CF_Always>::GetRHI();
 
-	bool bSetupBatchRenderingData = true;
+	FSamplerStateRHIRef BilinearClamp = TStaticSamplerState<SF_Bilinear,AM_Clamp,AM_Clamp,AM_Clamp>::GetRHI();
+
+	RHICmdList.SetStreamSource(0, VertexBuffer.VertexBufferRHI, sizeof(FSlateVertex), 0);
 
 	// Draw each element
 	for( int32 BatchIndex = 0; BatchIndex < RenderBatches.Num(); ++BatchIndex )
 	{
-		if( bSetupBatchRenderingData )
-		{
-			// Set the vertex buffer stream
-			RHICmdList.SetStreamSource(0, VertexBuffer.VertexBufferRHI, sizeof(FSlateVertex), 0);
-
-			bSetupBatchRenderingData = false;
-		}
-
 		const FSlateRenderBatch& RenderBatch = RenderBatches[BatchIndex];
 		const FSlateShaderResource* ShaderResource = RenderBatch.Texture;
 
+		const ESlateBatchDrawFlag::Type DrawFlags = RenderBatch.DrawFlags;
+		const ESlateDrawEffect::Type DrawEffects = RenderBatch.DrawEffects;
+		const ESlateShader::Type ShaderType = RenderBatch.ShaderType;
+		const FShaderParams& ShaderParams = RenderBatch.ShaderParams;
+		const bool bDrawDisabled = (DrawEffects & ESlateDrawEffect::DisabledEffect) != 0;
+
 		if( !RenderBatch.CustomDrawer.IsValid() && (!ShaderResource || ShaderResource->GetType() == ESlateShaderResource::Texture ) )
 		{
-			const ESlateBatchDrawFlag::Type DrawFlags = RenderBatch.DrawFlags;
-			const ESlateDrawEffect::Type DrawEffects = RenderBatch.DrawEffects;
-			const ESlateShader::Type ShaderType = RenderBatch.ShaderType;
-			const FShaderParams& ShaderParams = RenderBatch.ShaderParams;
-
 			FSlateElementPS* PixelShader = GetPixelShader( ShaderType, DrawEffects );
-
-			const bool bDrawDisabled = (DrawEffects & ESlateDrawEffect::DisabledEffect) != 0;
+	
 			const bool bUseTextureAlpha = (DrawEffects & ESlateDrawEffect::IgnoreTextureAlpha) == 0;
 
 			if( bDrawDisabled )
@@ -416,8 +367,6 @@ void FSlateRHIRenderingPolicy::DrawElements( const FIntPoint& InViewportSize, FS
 			}
 
 			PixelShader->SetShaderParams(RHICmdList, ShaderParams.PixelParams );
-			PixelShader->SetViewportSize(RHICmdList, InViewportSize );
-			PixelShader->SetDrawEffects(RHICmdList, DrawEffects );
 			PixelShader->SetDisplayGamma(RHICmdList, ( DrawFlags & ESlateBatchDrawFlag::NoGamma ) ? 1.0f : DisplayGamma );
 
 			check( RenderBatch.NumIndices > 0 );
@@ -428,7 +377,25 @@ void FSlateRHIRenderingPolicy::DrawElements( const FIntPoint& InViewportSize, FS
 		}
 		else if( ShaderResource && ShaderResource->GetType() == ESlateShaderResource::Material )
 		{
-			//bSetupBatchRenderingData = true;
+#if ALLOW_SLATE_MATERIALS
+			FMaterialRenderProxy* MaterialRenderProxy = ((FSlateMaterial*)ShaderResource)->GetRenderProxy();
+		
+			if( !SceneView )
+			{
+				SceneView = &CreateSceneView( SceneViewContext, BackBuffer, ViewProjectionMatrix );
+			}
+
+		
+			// Bind Shader here
+
+			VertexShader->SetShaderParameters(RHICmdList, ShaderParams.VertexParams);
+
+			check(RenderBatch.NumIndices > 0);
+
+			uint32 PrimitiveCount = RenderBatch.DrawPrimitiveType == ESlateDrawPrimitive::LineList ? RenderBatch.NumIndices / 2 : RenderBatch.NumIndices / 3;
+
+			RHICmdList.DrawIndexedPrimitive(IndexBuffer.IndexBufferRHI, GetRHIPrimitiveType(RenderBatch.DrawPrimitiveType), RenderBatch.VertexOffset, 0, RenderBatch.NumVertices, RenderBatch.IndexOffset, PrimitiveCount, 1);
+#endif
 		}
 		else if (RenderBatch.CustomDrawer.IsValid())
 		{
@@ -517,19 +484,3 @@ FSlateElementPS* FSlateRHIRenderingPolicy::GetPixelShader( ESlateShader::Type Sh
 
 	return PixelShader;
 }
-
-FSlateShaderResource* FSlateRHIRenderingPolicy::GetViewportResource( const ISlateViewport* InViewportInterface )
-{
-	if( InViewportInterface  )
-	{
-		return InViewportInterface->GetViewportRenderTargetTexture();
-	}
-
-	return NULL;
-}
-
-FSlateShaderResourceProxy* FSlateRHIRenderingPolicy::GetTextureResource( const FSlateBrush& Brush )
-{
-	check( IsThreadSafeForSlateRendering() );
-	return TextureManager->GetTexture( Brush );
-}
diff --git a/Engine/Source/Runtime/SlateRHIRenderer/Private/SlateRHIRenderingPolicy.h b/Engine/Source/Runtime/SlateRHIRenderer/Private/SlateRHIRenderingPolicy.h
--- a/Engine/Source/Runtime/SlateRHIRenderer/Private/SlateRHIRenderingPolicy.h
+++ b/Engine/Source/Runtime/SlateRHIRenderer/Private/SlateRHIRenderingPolicy.h
@@ -11,17 +11,46 @@ class FSlateRHIResourceManager;
  * Vertex buffer containing all Slate vertices
  * All vertices are added through RHILockVertexBuffer
  */
-class FSlateElementVertexBuffer : public FVertexBuffer
+template <typename VertexType>
+class TSlateElementVertexBuffer : public FVertexBuffer
 {
 public:
-	FSlateElementVertexBuffer();
-	~FSlateElementVertexBuffer();
+	TSlateElementVertexBuffer()
+		: BufferSize(0)
+		, BufferUsageSize(0)
+	{}
+
+	TSlateElementVertexBuffer( uint32 BufferSize )
+		: BufferSize(BufferSize)
+		, BufferUsageSize(0)
+	{}
+
+	~TSlateElementVertexBuffer() {};
 
 	/** Initializes the vertex buffers RHI resource. */
-	virtual void InitDynamicRHI();
+	virtual void InitDynamicRHI()
+	{
+		if( !IsValidRef(VertexBufferRHI) )
+		{
+			if( BufferSize == 0 )
+			{
+				BufferSize = 200 * sizeof(VertexType);
+			}
+
+			FRHIResourceCreateInfo CreateInfo;
+			VertexBufferRHI = RHICreateVertexBuffer( BufferSize, BUF_Dynamic, CreateInfo );
+
+			// Ensure the vertex buffer could be created
+			check(IsValidRef(VertexBufferRHI));
+		}
+	}
 
 	/** Releases the vertex buffers RHI resource. */
-	virtual void ReleaseDynamicRHI();
+	virtual void ReleaseDynamicRHI()
+	{
+		VertexBufferRHI.SafeRelease();
+		BufferSize = 0;
+	}
 
 	/** Returns a friendly name for this buffer. */
 	virtual FString GetFriendlyName() const { return TEXT("SlateElementVertices"); }
@@ -36,10 +65,56 @@ public:
 	void ResetBufferUsage() { BufferUsageSize = 0; }
 
 	/** Fills the buffer with slate vertices */
-	void FillBuffer( const TArray<FSlateVertex>& InVertices, bool bShrinkToFit );
+	void FillBuffer( const TArray<VertexType>& InVertices, bool bShrinkToFit )
+	{
+		check( IsInRenderingThread() );
+
+		if( InVertices.Num() )
+		{
+			uint32 NumVertices = InVertices.Num();
+
+	#if !SLATE_USE_32BIT_INDICES
+			// make sure our index buffer can handle this
+			checkf(NumVertices < 0xFFFF, TEXT("Slate vertex buffer is too large (%d) to work with uint16 indices"), NumVertices);
+	#endif
+
+			uint32 RequiredBufferSize = NumVertices*sizeof(VertexType);
+
+			// resize if needed
+			if( RequiredBufferSize > GetBufferSize() || bShrinkToFit )
+			{
+				// Use array resize techniques for the vertex buffer
+				ResizeBuffer( InVertices.GetAllocatedSize() );
+			}
+
+			BufferUsageSize += RequiredBufferSize;
+
+			void* VerticesPtr = RHILockVertexBuffer( VertexBufferRHI, 0, RequiredBufferSize, RLM_WriteOnly );
+
+			FMemory::Memcpy( VerticesPtr, InVertices.GetData(), RequiredBufferSize );
+
+			RHIUnlockVertexBuffer( VertexBufferRHI );
+		}
+	}
+
 private:
 	/** Resizes the buffer to the passed in size.  Preserves internal data*/
-	void ResizeBuffer( uint32 NewSizeBytes );
+	void ResizeBuffer( uint32 NewSizeBytes )
+	{
+		check( IsInRenderingThread() );
+
+		if( NewSizeBytes != 0 && NewSizeBytes != BufferSize)
+		{
+			VertexBufferRHI.SafeRelease();
+
+			FRHIResourceCreateInfo CreateInfo;
+			VertexBufferRHI = RHICreateVertexBuffer(NewSizeBytes, BUF_Dynamic, CreateInfo);
+
+			check(IsValidRef(VertexBufferRHI));
+
+			BufferSize = NewSizeBytes;
+		}
+	}
 
 private:
 	/** The size of the buffer in bytes. */
@@ -49,8 +124,8 @@ private:
 	uint32 BufferUsageSize;
 
 	/** Hidden copy methods. */
-	FSlateElementVertexBuffer( const FSlateElementVertexBuffer& );
-	void operator=(const FSlateElementVertexBuffer& );
+	TSlateElementVertexBuffer( const TSlateElementVertexBuffer& );
+	void operator=(const TSlateElementVertexBuffer& );
 };
 
 class FSlateElementIndexBuffer : public FIndexBuffer
@@ -95,32 +170,18 @@ private:
 
 };
 
-class FSlateRenderTarget : public FRenderTarget
-{
-public:
-	FSlateRenderTarget( FTexture2DRHIRef& InRenderTargetTexture, FIntPoint InSizeXY )
-		: SizeXY( InSizeXY )
-	{
-		RenderTargetTextureRHI = InRenderTargetTexture;
-	}
-	virtual FIntPoint GetSizeXY() const override { return SizeXY; }
-private:
-	FIntPoint SizeXY;
-};
-
 class FSlateRHIRenderingPolicy : public FSlateRenderingPolicy
 {
 public:
-	FSlateRHIRenderingPolicy( TSharedPtr<FSlateFontCache> InFontCache, TSharedRef<FSlateRHIResourceManager> InTextureManager );
+	FSlateRHIRenderingPolicy( TSharedPtr<FSlateFontCache> InFontCache, TSharedRef<FSlateRHIResourceManager> InResourceManager );
 	~FSlateRHIRenderingPolicy();
 	
 	virtual void UpdateBuffers( const FSlateWindowElementList& WindowElementList ) override;
-	virtual void DrawElements( const FIntPoint& InViewportSize, FSlateRenderTarget& BackBuffer, const FMatrix& ViewProjectionMatrix, const TArray<FSlateRenderBatch>& RenderBatches );
+	virtual void DrawElements( const FIntPoint& InViewportSize, class FSlateBackBuffer& BackBuffer, const FMatrix& ViewProjectionMatrix, const TArray<FSlateRenderBatch>& RenderBatches );
+
+	TSharedPtr<FSlateFontCache> GetFontCache() { return FontCache; }
+	TSharedRef<FSlateShaderResourceManager> GetResourceManager() { return ResourceManager; }
 
-	class FSlateShaderResource* GetViewportResource( const ISlateViewport* InViewportInterface );
-	virtual class FSlateShaderResourceProxy* GetTextureResource( const FSlateBrush& Brush ) override;
-	TSharedPtr<FSlateFontCache>& GetFontCache() { return FontCache; }
-	
 	void InitResources();
 	void ReleaseResources();
 
@@ -139,13 +200,15 @@ private:
 	/** Global shader states */
 	static FGlobalBoundShaderState NormalShaderStates[4][2 /* UseTextureAlpha */];
 	static FGlobalBoundShaderState DisabledShaderStates[4][2 /* UseTextureAlpha */];
+	static FGlobalBoundShaderState MaterialShaderState;
 
 	/** Buffers used for rendering */
-	FSlateElementVertexBuffer VertexBuffers[2];
+	TSlateElementVertexBuffer<FSlateVertex> VertexBuffers[2];
 	FSlateElementIndexBuffer IndexBuffers[2];
 
-	TSharedRef<FSlateRHIResourceManager> TextureManager;
+	TSharedRef<FSlateRHIResourceManager> ResourceManager;
 	TSharedPtr<FSlateFontCache> FontCache;
+
 	uint8 CurrentBufferIndex;
 	
 	/** If we should shrink resources that are no longer used (do not set this from the game thread)*/
diff --git a/Engine/Source/Runtime/SlateRHIRenderer/Private/SlateRHIResourceManager.cpp b/Engine/Source/Runtime/SlateRHIRenderer/Private/SlateRHIResourceManager.cpp
--- a/Engine/Source/Runtime/SlateRHIRenderer/Private/SlateRHIResourceManager.cpp
+++ b/Engine/Source/Runtime/SlateRHIRenderer/Private/SlateRHIResourceManager.cpp
@@ -30,11 +30,33 @@ FDynamicTextureResource::~FDynamicTextureResource()
 }
 
 
-TSharedPtr<FDynamicTextureResource> FDynamicResourceMap::GetResource( FName ResourceName, UObject* ResourceObject ) const
+FSlateMaterialResource::FSlateMaterialResource(UMaterialInterface& InMaterial, const FVector2D& InImageSize)
+	: MaterialObject(&InMaterial)
+	, Proxy(new FSlateShaderResourceProxy)
+	, SlateMaterial(new FSlateMaterial(InMaterial.GetRenderProxy(false), InImageSize))
+{
+	Proxy->ActualSize = InImageSize.IntPoint();
+	Proxy->Resource = SlateMaterial;
+}
+
+FSlateMaterialResource::~FSlateMaterialResource()
 {
-	if( ResourceObject )
+	if (Proxy)
+	{
+		delete Proxy;
+	}
+
+	if (SlateMaterial)
+	{
+		delete SlateMaterial;
+	}
+}
+
+TSharedPtr<FDynamicTextureResource> FDynamicResourceMap::GetTextureResource( FName ResourceName, UTexture2D* TextureObject ) const
+{
+	if( TextureObject )
 	{
-		return DynamicResourceObjectMap.FindRef( ResourceObject );
+		return DynamicResourceObjectMap.FindRef( TextureObject );
 	}
 	else
 	{
@@ -43,11 +65,18 @@ TSharedPtr<FDynamicTextureResource> FDynamicResourceMap::GetResource( FName Reso
 }
 
 
-void FDynamicResourceMap::AddResource( FName ResourceName, UObject* ResourceObject, TSharedRef<FDynamicTextureResource> InResource )
+TSharedPtr<FSlateMaterialResource> FDynamicResourceMap::GetMaterialResource( UMaterialInterface* Material ) const
 {
-	if( ResourceObject )
+	return MaterialResourceMap.FindRef( Material );
+}
+
+
+void FDynamicResourceMap::AddTextureResource( FName ResourceName, UTexture2D* TextureObject, TSharedRef<FDynamicTextureResource> InResource )
+{
+	if( TextureObject )
 	{
-		DynamicResourceObjectMap.Add( ResourceObject, InResource );
+		check( TextureObject == InResource->TextureObject );
+		DynamicResourceObjectMap.Add( TextureObject, InResource );
 	}
 	else
 	{
@@ -56,11 +85,17 @@ void FDynamicResourceMap::AddResource( FName ResourceName, UObject* ResourceObje
 
 }
 
-void FDynamicResourceMap::RemoveResource( FName ResourceName, UObject* ResourceObject )
+void FDynamicResourceMap::AddMaterialResource( UMaterialInterface* Material, TSharedRef<FSlateMaterialResource> InMaterialResource )
+{
+	check( Material == InMaterialResource->MaterialObject );
+	MaterialResourceMap.Add( Material, InMaterialResource );
+}
+
+void FDynamicResourceMap::RemoveTextureResource( FName ResourceName, UTexture2D* TextureObject )
 {
-	if( ResourceObject )
+	if( TextureObject )
 	{
-		DynamicResourceObjectMap.Remove( ResourceObject );
+		DynamicResourceObjectMap.Remove( TextureObject );
 	}
 	else
 	{
@@ -68,10 +103,16 @@ void FDynamicResourceMap::RemoveResource( FName ResourceName, UObject* ResourceO
 	}
 }
 
+void FDynamicResourceMap::RemoveMaterialResource( UMaterialInterface* Material )
+{
+	MaterialResourceMap.Remove( Material );
+}
+
 void FDynamicResourceMap::Empty()
 {
 	DynamicNativeTextureMap.Empty();
 	DynamicResourceObjectMap.Empty();
+	MaterialResourceMap.Empty();
 }
 
 void FDynamicResourceMap::ReleaseResources()
@@ -106,8 +147,19 @@ void FDynamicResourceMap::AddReferencedObjects(FReferenceCollector& Collector)
 			Collector.AddReferencedObject(Resource->TextureObject);
 		}
 	}
+
+
+	for (TMap<UMaterialInterface*, TSharedPtr<FSlateMaterialResource> >::TIterator It(MaterialResourceMap); It; ++It)
+	{
+		TSharedPtr<FSlateMaterialResource>& Resource = It.Value();
+		if (Resource->MaterialObject)
+		{
+			Collector.AddReferencedObject(Resource->MaterialObject);
+		}
+	}
 }
 
+
 FSlateRHIResourceManager::FSlateRHIResourceManager()
 {
 	if( GIsEditor )
@@ -311,8 +363,7 @@ FSlateShaderResourceProxy* FSlateRHIResourceManager::GenerateTextureResource( co
 	return NewProxy;
 }
 
-
-FSlateShaderResourceProxy* FSlateRHIResourceManager::GetTexture( const FSlateBrush& InBrush )
+FSlateShaderResourceProxy* FSlateRHIResourceManager::GetShaderResource( const FSlateBrush& InBrush )
 {
 	check( IsThreadSafeForSlateRendering() );
 
@@ -321,6 +372,10 @@ FSlateShaderResourceProxy* FSlateRHIResourceManager::GetTexture( const FSlateBru
 	{
 		Texture = ResourceMap.FindRef( InBrush.GetResourceName() );
 	}
+	else if (InBrush.GetResourceObject() && InBrush.GetResourceObject()->IsA<UMaterialInterface>())
+	{
+		Texture = GetMaterialResource(InBrush);
+	}
 	else if( InBrush.IsDynamicallyLoaded() || ( InBrush.HasUObject() ) )
 	{
 		Texture = GetDynamicTextureResource( InBrush );
@@ -332,7 +387,7 @@ FSlateShaderResourceProxy* FSlateRHIResourceManager::GetTexture( const FSlateBru
 TSharedPtr<FDynamicTextureResource> FSlateRHIResourceManager::MakeDynamicTextureResource( FName ResourceName, uint32 Width, uint32 Height, const TArray< uint8 >& Bytes )
 {
 	// Bail out if we already have this texture loaded
-	TSharedPtr<FDynamicTextureResource> TextureResource = DynamicResourceMap.GetResource( ResourceName, nullptr );
+	TSharedPtr<FDynamicTextureResource> TextureResource = DynamicResourceMap.GetTextureResource( ResourceName, nullptr );
 	if( TextureResource.IsValid() )
 	{
 		return TextureResource;
@@ -345,7 +400,7 @@ TSharedPtr<FDynamicTextureResource> FSlateRHIResourceManager::MakeDynamicTexture
 	TextureResource = InitializeDynamicTextureResource( TextureStorage, NULL );
 
 	// Map the new resource so we don't have to load again
-	DynamicResourceMap.AddResource( ResourceName, nullptr, TextureResource.ToSharedRef() );
+	DynamicResourceMap.AddTextureResource( ResourceName, nullptr, TextureResource.ToSharedRef() );
 
 	return TextureResource;
 }
@@ -366,7 +421,7 @@ TSharedPtr<FDynamicTextureResource> FSlateRHIResourceManager::MakeDynamicTexture
 	bool bSucceeded = false;
 	if( bHasUTexture || InTextureObject != NULL )
 	{
-		TextureResource = DynamicResourceMap.GetResource( ResourceName, InTextureObject );
+		TextureResource = DynamicResourceMap.GetTextureResource( ResourceName, InTextureObject );
 		if( TextureResource.IsValid() )
 		{
 			// Bail out of the resource is already loaded
@@ -391,7 +446,7 @@ TSharedPtr<FDynamicTextureResource> FSlateRHIResourceManager::MakeDynamicTexture
 	}
 	else if( bIsDynamicallyLoaded )
 	{
-		TextureResource = DynamicResourceMap.GetResource(ResourceName, nullptr );
+		TextureResource = DynamicResourceMap.GetTextureResource(ResourceName, nullptr );
 		if ( TextureResource.IsValid() )
 		{
 			// Bail out of the resource is already loaded
@@ -421,11 +476,11 @@ TSharedPtr<FDynamicTextureResource> FSlateRHIResourceManager::MakeDynamicTexture
 
 	if( !bUsingDeprecatedUTexturePath )
 	{
-		DynamicResourceMap.AddResource(ResourceName, TextureObject, TextureResource.ToSharedRef());
+		DynamicResourceMap.AddTextureResource(ResourceName, TextureObject, TextureResource.ToSharedRef());
 	}
 	else
 	{
-		DynamicResourceMap.AddResource(ResourceName, nullptr, TextureResource.ToSharedRef());
+		DynamicResourceMap.AddTextureResource(ResourceName, nullptr, TextureResource.ToSharedRef());
 	}
 
 
@@ -496,7 +551,7 @@ FSlateShaderResourceProxy* FSlateRHIResourceManager::GetDynamicTextureResource(
 	const FName ResourceName = InBrush.GetResourceName();
 	if ( ResourceName.IsValid() && ResourceName != NAME_None )
 	{
-		TSharedPtr<FDynamicTextureResource> TextureResource = DynamicResourceMap.GetResource( ResourceName, InBrush.GetResourceObject() );
+		TSharedPtr<FDynamicTextureResource> TextureResource = DynamicResourceMap.GetTextureResource( ResourceName, Cast<UTexture2D>( InBrush.GetResourceObject() ) );
 
 		if( !TextureResource.IsValid() )
 		{
@@ -536,6 +591,31 @@ FSlateShaderResourceProxy* FSlateRHIResourceManager::GetDynamicTextureResource(
 	return  NULL;
 }
 
+FSlateShaderResourceProxy* FSlateRHIResourceManager::GetMaterialResource(const FSlateBrush& InBrush)
+{
+	check(IsInGameThread());
+
+	const FName ResourceName = InBrush.GetResourceName();
+
+	UMaterialInterface* Material = CastChecked<UMaterialInterface>(InBrush.GetResourceObject());
+
+	TSharedPtr<FSlateMaterialResource> MaterialResource = DynamicResourceMap.GetMaterialResource(Material);
+	if (!MaterialResource.IsValid())
+	{
+		TSharedRef<FSlateMaterialResource> NewMaterialResource = MakeShareable(new FSlateMaterialResource(*Material, InBrush.ImageSize));
+		DynamicResourceMap.AddMaterialResource( Material, NewMaterialResource );
+
+		MaterialResource = NewMaterialResource;
+	}
+	else
+	{
+		// Keep the resource up to date
+		MaterialResource->SlateMaterial->SetResource(Material->GetRenderProxy(false, false));
+	}
+
+	return MaterialResource->Proxy;
+}
+
 bool FSlateRHIResourceManager::ContainsTexture( const FName& ResourceName ) const
 {
 	return ResourceMap.Contains( ResourceName );
@@ -549,7 +629,10 @@ void FSlateRHIResourceManager::ReleaseDynamicResource( const FSlateBrush& InBrus
 	if( InBrush.HasUObject() || InBrush.IsDynamicallyLoaded() )
 	{
 		FName ResourceName = InBrush.GetResourceName();
-		TSharedPtr<FDynamicTextureResource> TextureResource = DynamicResourceMap.GetResource( ResourceName, InBrush.GetResourceObject() );
+
+		UObject* ResourceObject = InBrush.GetResourceObject();
+
+		TSharedPtr<FDynamicTextureResource> TextureResource = DynamicResourceMap.GetTextureResource( ResourceName, Cast<UTexture2D>( ResourceObject ) );
 
 		if( TextureResource.IsValid() )
 		{
@@ -560,13 +643,20 @@ void FSlateRHIResourceManager::ReleaseDynamicResource( const FSlateBrush& InBrus
 			BeginReleaseResource( TextureResource->RHIRefTexture );
 
 			//remove it from the texture map
-			DynamicResourceMap.RemoveResource( ResourceName, InBrush.GetResourceObject() );
+			DynamicResourceMap.RemoveTextureResource( ResourceName, TextureResource->TextureObject );
+
+			TextureResource->TextureObject = nullptr;
 
 			// Add the resource to the free list so it can be reused
 			DynamicTextureFreeList.Add( TextureResource );
-
+			
 			DEC_DWORD_STAT_BY( STAT_SlateNumDynamicTextures, 1 );
 		}
+		else
+		{
+			UMaterialInterface* Material = Cast<UMaterialInterface>( ResourceObject );
+			DynamicResourceMap.RemoveMaterialResource( Material );
+		}
 	}
 }
 
@@ -858,4 +948,4 @@ TSharedRef<SWidget> FSlateRHIResourceManager::CreateTextureDisplayWidget()
 	}
 
 	return AtlasVisualizer->MakeVisualizerWidget();
-}
\ No newline at end of file
+}
diff --git a/Engine/Source/Runtime/SlateRHIRenderer/Private/SlateRHIResourceManager.h b/Engine/Source/Runtime/SlateRHIRenderer/Private/SlateRHIResourceManager.h
--- a/Engine/Source/Runtime/SlateRHIRenderer/Private/SlateRHIResourceManager.h
+++ b/Engine/Source/Runtime/SlateRHIRenderer/Private/SlateRHIResourceManager.h
@@ -18,15 +18,33 @@ struct FDynamicTextureResource
 	FSlateTexture2DRHIRef* RHIRefTexture;
 };
 
+struct FSlateMaterialResource
+{
+	/** UObject material being used */
+	class UMaterialInterface* MaterialObject;
+	/** Slate proxy used for batching the material */
+	FSlateShaderResourceProxy* Proxy;
+	/** Slate material shader resource.  Resource used to render the object*/
+	FSlateMaterial* SlateMaterial;
+
+	FSlateMaterialResource(UMaterialInterface& InMaterial, const FVector2D& InImageSize);
+	~FSlateMaterialResource();
+};
 
 struct FDynamicResourceMap : public FGCObject
 {
 public:
-	TSharedPtr<FDynamicTextureResource> GetResource( FName ResourceName, UObject* ResourceObject ) const;
+	TSharedPtr<FDynamicTextureResource> GetTextureResource( FName ResourceName, UTexture2D* TextureObject ) const;
+
+	TSharedPtr<FSlateMaterialResource> GetMaterialResource( UMaterialInterface* Material ) const;
+
+	void AddTextureResource( FName ResourceName, UTexture2D* TextureObject, TSharedRef<FDynamicTextureResource> InResource );
+
+	void RemoveTextureResource( FName ResourceName, UTexture2D* TextureObject );
 
-	void AddResource( FName ResourceName, UObject* ResourceObject, TSharedRef<FDynamicTextureResource> InResource );
+	void AddMaterialResource( UMaterialInterface* Material, TSharedRef<FSlateMaterialResource> InMaterialResource );
 
-	void RemoveResource( FName ResourceName, UObject* ResourceObject );
+	void RemoveMaterialResource( UMaterialInterface* Material );
 
 	void Empty();
 
@@ -39,6 +57,9 @@ private:
 	TMap<FName, TSharedPtr<FDynamicTextureResource> > DynamicNativeTextureMap;
 
 	TMap<UObject*, TSharedPtr<FDynamicTextureResource> > DynamicResourceObjectMap;
+
+	/** Map of all material resources */
+	TMap<UMaterialInterface*, TSharedPtr<FSlateMaterialResource> > MaterialResourceMap;
 };
 
 
@@ -71,10 +92,8 @@ public:
 	 */
 	void UpdateTextureAtlases();
 
-	/**
-	 * Returns a texture associated with the passed in name. 
-	 */
-	virtual FSlateShaderResourceProxy* GetTexture( const FSlateBrush& InBrush );
+	/** FSlateShaderResourceManager interface */
+	virtual FSlateShaderResourceProxy* GetShaderResource( const FSlateBrush& InBrush ) override;
 	
 	/**
 	 * Makes a dynamic texture resource and begins use of it
@@ -172,6 +191,13 @@ private:
 	 */
 	FSlateShaderResourceProxy* GetDynamicTextureResource( const FSlateBrush& InBrush );
 
+	/**
+	 * Returns a rendering resource for a material
+	 *
+	 * @param InBrush	Slate brush for the material
+	 */
+	FSlateShaderResourceProxy* GetMaterialResource( const FSlateBrush& InBrush );
+
 private:
 	/** Map of all active dynamic resources being used by brushes */
 	FDynamicResourceMap DynamicResourceMap;
diff --git a/Engine/Source/Runtime/SlateRHIRenderer/Private/SlateShaders.h b/Engine/Source/Runtime/SlateRHIRenderer/Private/SlateShaders.h
--- a/Engine/Source/Runtime/SlateRHIRenderer/Private/SlateShaders.h
+++ b/Engine/Source/Runtime/SlateRHIRenderer/Private/SlateShaders.h
@@ -87,9 +87,7 @@ public:
 	{
 		TextureParameter.Bind( Initializer.ParameterMap, TEXT("ElementTexture"));
 		TextureParameterSampler.Bind( Initializer.ParameterMap, TEXT("ElementTextureSampler"));
-		DrawEffectsParameter.Bind( Initializer.ParameterMap, TEXT("DrawEffects"));
 		ShaderParams.Bind( Initializer.ParameterMap, TEXT("ShaderParams"));
-		ViewportSize.Bind(Initializer.ParameterMap, TEXT("ViewportSize"));
 		DisplayGamma.Bind( Initializer.ParameterMap,TEXT("DisplayGamma"));
 	}
 
@@ -105,16 +103,6 @@ public:
 		SetTextureParameter(RHICmdList, GetPixelShader(), TextureParameter, TextureParameterSampler, SamplerState, InTexture );
 	}
 
-	/**
-	 * Sets any draw effect flags used by the shader
-	 *
-	 * @param InDrawEffects		A mask of draw effects that the pixel shader should use.
-	 */
-	void SetDrawEffects(FRHICommandList& RHICmdList, ESlateDrawEffect::Type InDrawEffects )
-	{
-		SetShaderValue(RHICmdList, GetPixelShader(), DrawEffectsParameter, InDrawEffects );
-	}
-
 	/**
 	 * Sets shader params used by the shader
 	 * 
@@ -125,14 +113,6 @@ public:
 		SetShaderValue(RHICmdList, GetPixelShader(), ShaderParams, InShaderParams );
 	}
 
-	/**
-	 * Set the viewport size.
-	 */
-	void SetViewportSize(FRHICommandList& RHICmdList, const FVector2D& InViewportSize)
-	{
-		SetShaderValue(RHICmdList, GetPixelShader(), ViewportSize, InViewportSize);
-	}
-
 	/**
 	 * Sets the display gamma.
  	 *
@@ -150,8 +130,6 @@ public:
 		Ar << TextureParameter;
 		Ar << TextureParameterSampler;
 		Ar << ShaderParams;
-		Ar << ViewportSize;
-		Ar << DrawEffectsParameter;
 		Ar << DisplayGamma;
 
 		return bShaderHasOutdatedParameters;
@@ -161,8 +139,6 @@ private:
 	FShaderResourceParameter TextureParameter;
 	FShaderResourceParameter TextureParameterSampler;
 	FShaderParameter ShaderParams;
-	FShaderParameter ViewportSize;
-	FShaderParameter DrawEffectsParameter;
 	FShaderParameter DisplayGamma;
 };
 
@@ -196,6 +172,7 @@ public:
 		OutEnvironment.SetDefine( TEXT("DRAW_DISABLED_EFFECT"), (uint32)(bDrawDisabledEffect ? 1 : 0) );
 		OutEnvironment.SetDefine( TEXT("USE_TEXTURE_ALPHA"), (uint32)(bUseTextureAlpha ? 1 : 0) );
 		OutEnvironment.SetDefine( TEXT("COLOR_VISION_DEFICIENCY_TYPE"), GSlateShaderColorVisionDeficiencyType );
+		OutEnvironment.SetDefine( TEXT("USE_MATERIALS"), (uint32)0 );
 
 		FSlateElementPS::ModifyCompilationEnvironment( Platform, OutEnvironment );
 	}
