Added ability to visualize overdraw in Slate.

diff --git a/Engine/Shaders/SlateElementPixelShader.usf b/Engine/Shaders/SlateElementPixelShader.usf
--- a/Engine/Shaders/SlateElementPixelShader.usf
+++ b/Engine/Shaders/SlateElementPixelShader.usf
@@ -298,3 +298,12 @@ float4 Main( VertexOut VIn ) : SV_Target0
 	return OutColor;
 }
 
+float4 DebugOverdrawMain( VertexOut VIn ) : SV_Target0
+{
+#if !DEBUG_CLIPPING
+	// Clip pixels which are outside of the clipping rect
+	clip( float4( VIn.WorldPosition.x - VIn.ClipCoords.x, VIn.ClipCoords.z - VIn.WorldPosition.x, VIn.WorldPosition.y - VIn.ClipCoords.y, VIn.ClipCoords.w - VIn.WorldPosition.y ) );
+#endif
+
+	return .1;
+}
diff --git a/Engine/Source/Runtime/SlateRHIRenderer/Private/SlateRHIRenderer.cpp b/Engine/Source/Runtime/SlateRHIRenderer/Private/SlateRHIRenderer.cpp
--- a/Engine/Source/Runtime/SlateRHIRenderer/Private/SlateRHIRenderer.cpp
+++ b/Engine/Source/Runtime/SlateRHIRenderer/Private/SlateRHIRenderer.cpp
@@ -457,6 +457,9 @@ void FSlateRHIRenderer::DrawWindow_RenderThread( const FViewportInfo& ViewportIn
 			RHIClear( false, FLinearColor::White, false, 0.0f, true, 0x00, FIntRect());
 		}
 
+#if DEBUG_OVERDRAW
+		RHIClear(true, FLinearColor::Black, false, 0.0f, true, 0x00, FIntRect());
+#endif
 		if( WindowElementList.GetRenderBatches().Num() > 0 )
 		{
 			FSlateBackBuffer BackBufferTarget( BackBuffer, FIntPoint( ViewportInfo.Width, ViewportInfo.Height ) );
diff --git a/Engine/Source/Runtime/SlateRHIRenderer/Private/SlateRHIRenderer.h b/Engine/Source/Runtime/SlateRHIRenderer/Private/SlateRHIRenderer.h
--- a/Engine/Source/Runtime/SlateRHIRenderer/Private/SlateRHIRenderer.h
+++ b/Engine/Source/Runtime/SlateRHIRenderer/Private/SlateRHIRenderer.h
@@ -13,6 +13,8 @@ class FSlateRHIRenderingPolicy;
 const uint32 NumDrawBuffers = 3;
 #endif
 
+// Enable to visualize overdraw in Slate
+#define DEBUG_OVERDRAW 0
 
 /** Resource for crash buffer editor capturing */
 class FSlateCrashReportResource : public FRenderResource
diff --git a/Engine/Source/Runtime/SlateRHIRenderer/Private/SlateRHIRenderingPolicy.cpp b/Engine/Source/Runtime/SlateRHIRenderer/Private/SlateRHIRenderingPolicy.cpp
--- a/Engine/Source/Runtime/SlateRHIRenderer/Private/SlateRHIRenderingPolicy.cpp
+++ b/Engine/Source/Runtime/SlateRHIRenderer/Private/SlateRHIRenderingPolicy.cpp
@@ -295,11 +295,15 @@ void FSlateRHIRenderingPolicy::DrawElements( const FIntPoint& InViewportSize, FS
 			VertexShader->SetViewProjection(RHICmdList, ViewProjectionMatrix );
 			VertexShader->SetShaderParameters(RHICmdList, ShaderParams.VertexParams );
 
+#if !DEBUG_OVERDRAW
 			RHICmdList.SetBlendState(
 				(RenderBatch.DrawFlags & ESlateBatchDrawFlag::NoBlending)
 				? TStaticBlendState<>::GetRHI()
 				: TStaticBlendState<CW_RGBA,BO_Add,BF_SourceAlpha,BF_InverseSourceAlpha,BO_Add,BF_Zero,BF_One>::GetRHI()
 				);
+#else
+			RHICmdList.SetBlendState(TStaticBlendState<CW_RGB, BO_Add, BF_One, BF_One, BO_Add, BF_Zero, BF_InverseSourceAlpha>::GetRHI());
+#endif
 
 			// Disable stencil testing by default
 			RHICmdList.SetDepthStencilState(DSOff);
@@ -413,6 +417,7 @@ FSlateElementPS* FSlateRHIRenderingPolicy::GetPixelShader( ESlateShader::Type Sh
 {
 	FSlateElementPS* PixelShader = NULL;
 
+#if !DEBUG_OVERDRAW
 	const bool bDrawDisabled = (DrawEffects & ESlateDrawEffect::DisabledEffect) != 0;
 	const bool bUseTextureAlpha = (DrawEffects & ESlateDrawEffect::IgnoreTextureAlpha) == 0;
 	if( bDrawDisabled )
@@ -481,6 +486,9 @@ FSlateElementPS* FSlateRHIRenderingPolicy::GetPixelShader( ESlateShader::Type Sh
 			break;
 		}
 	}
+#else
+	PixelShader = *TShaderMapRef<FSlateDebugOverdrawPS>( GetGlobalShaderMap() );
+#endif
 
 	return PixelShader;
 }
diff --git a/Engine/Source/Runtime/SlateRHIRenderer/Private/SlateShaders.cpp b/Engine/Source/Runtime/SlateRHIRenderer/Private/SlateShaders.cpp
--- a/Engine/Source/Runtime/SlateRHIRenderer/Private/SlateShaders.cpp
+++ b/Engine/Source/Runtime/SlateRHIRenderer/Private/SlateShaders.cpp
@@ -9,6 +9,8 @@ uint32 GSlateShaderColorVisionDeficiencyType = 0;
 
 IMPLEMENT_SHADER_TYPE(, FSlateElementVS, TEXT("SlateVertexShader"),TEXT("Main"),SF_Vertex);
 
+IMPLEMENT_SHADER_TYPE(, FSlateDebugOverdrawPS, TEXT("SlateElementPixelShader"), TEXT("DebugOverdrawMain"), SF_Pixel );
+
 #define IMPLEMENT_SLATE_PIXELSHADER_TYPE(ShaderType, bDrawDisabledEffect, bUseTextureAlpha) \
 	typedef TSlateElementPS<ESlateShader::ShaderType,bDrawDisabledEffect,bUseTextureAlpha> TSlateElementPS##ShaderType##bDrawDisabledEffect##bUseTextureAlpha; \
 	IMPLEMENT_SHADER_TYPE(template<>,TSlateElementPS##ShaderType##bDrawDisabledEffect##bUseTextureAlpha,TEXT("SlateElementPixelShader"),TEXT("Main"),SF_Pixel);
diff --git a/Engine/Source/Runtime/SlateRHIRenderer/Private/SlateShaders.h b/Engine/Source/Runtime/SlateRHIRenderer/Private/SlateShaders.h
--- a/Engine/Source/Runtime/SlateRHIRenderer/Private/SlateShaders.h
+++ b/Engine/Source/Runtime/SlateRHIRenderer/Private/SlateShaders.h
@@ -186,5 +186,36 @@ private:
 
 };
 
+/** 
+ * Pixel shader for debugging Slate overdraw
+ */
+class FSlateDebugOverdrawPS : public FSlateElementPS
+{	
+	DECLARE_SHADER_TYPE( FSlateDebugOverdrawPS, Global );
+public:
+	/** Indicates that this shader should be cached */
+	static bool ShouldCache( EShaderPlatform Platform ) 
+	{ 
+		return true; 
+	}
+
+	FSlateDebugOverdrawPS()
+	{
+	}
+
+	/** Constructor.  Binds all parameters used by the shader */
+	FSlateDebugOverdrawPS( const ShaderMetaType::CompiledShaderInitializerType& Initializer )
+		: FSlateElementPS( Initializer )
+	{
+	}
+
+
+	virtual bool Serialize( FArchive& Ar )
+	{
+		return FSlateElementPS::Serialize( Ar );
+	}
+private:
+};
+
 /** The simple element vertex declaration. */
 extern TGlobalResource<FSlateVertexDeclaration> GSlateVertexDeclaration;
