Fixed non unity build.

diff --git a/Engine/Source/Editor/DetailCustomizations/Private/SlateBrushCustomization.cpp b/Engine/Source/Editor/DetailCustomizations/Private/SlateBrushCustomization.cpp
--- a/Engine/Source/Editor/DetailCustomizations/Private/SlateBrushCustomization.cpp
+++ b/Engine/Source/Editor/DetailCustomizations/Private/SlateBrushCustomization.cpp
@@ -2,6 +2,7 @@
 
 #include "DetailCustomizationsPrivatePCH.h"
 #include "SlateBrushCustomization.h"
+#include "AssetData.h"
 
 /**
  * Slate Brush Preview widget
diff --git a/Engine/Source/Runtime/Core/Private/Modules/ModuleManager.cpp b/Engine/Source/Runtime/Core/Private/Modules/ModuleManager.cpp
--- a/Engine/Source/Runtime/Core/Private/Modules/ModuleManager.cpp
+++ b/Engine/Source/Runtime/Core/Private/Modules/ModuleManager.cpp
@@ -4,6 +4,7 @@
 #include "ModuleManager.h"
 #include "ModuleVersion.h"
 #include "EngineBuildSettings.h"
+#include "UProjectInfo.h"
 
 DEFINE_LOG_CATEGORY_STATIC(LogModuleManager, Log, All);
 
diff --git a/Engine/Source/Runtime/SlateRHIRenderer/Private/SlateRHIRenderer.cpp b/Engine/Source/Runtime/SlateRHIRenderer/Private/SlateRHIRenderer.cpp
--- a/Engine/Source/Runtime/SlateRHIRenderer/Private/SlateRHIRenderer.cpp
+++ b/Engine/Source/Runtime/SlateRHIRenderer/Private/SlateRHIRenderer.cpp
@@ -14,19 +14,6 @@ namespace CrashTrackerConstants
 }
 
 
-class FSlateBackBuffer : public FRenderTarget
-{
-public:
-	FSlateBackBuffer(FTexture2DRHIRef& InRenderTargetTexture, FIntPoint InSizeXY)
-		: SizeXY(InSizeXY)
-	{
-		RenderTargetTextureRHI = InRenderTargetTexture;
-	}
-	virtual FIntPoint GetSizeXY() const override { return SizeXY; }
-private:
-	FIntPoint SizeXY;
-};
-
 // Defines the maximum size that a slate viewport will create
 #define MAX_VIEWPORT_SIZE 16384
 
diff --git a/Engine/Source/Runtime/SlateRHIRenderer/Private/SlateRHIRenderer.h b/Engine/Source/Runtime/SlateRHIRenderer/Private/SlateRHIRenderer.h
--- a/Engine/Source/Runtime/SlateRHIRenderer/Private/SlateRHIRenderer.h
+++ b/Engine/Source/Runtime/SlateRHIRenderer/Private/SlateRHIRenderer.h
@@ -16,6 +16,19 @@ const uint32 NumDrawBuffers = 3;
 // Enable to visualize overdraw in Slate
 #define DEBUG_OVERDRAW 0
 
+class FSlateBackBuffer : public FRenderTarget
+{
+public:
+	FSlateBackBuffer(FTexture2DRHIRef& InRenderTargetTexture, FIntPoint InSizeXY)
+		: SizeXY(InSizeXY)
+	{
+		RenderTargetTextureRHI = InRenderTargetTexture;
+	}
+	virtual FIntPoint GetSizeXY() const override { return SizeXY; }
+private:
+	FIntPoint SizeXY;
+};
+
 /** Resource for crash buffer editor capturing */
 class FSlateCrashReportResource : public FRenderResource
 {
