Work in progress: Live streaming support for editor and games

This commit adds the framework APIs, editor and rendering features needed to support live streaming.  Along with this, we're working on a new plugin that adds Twitch.tv support (waiting for legal approval to commit.)

- Game live streaming
       - Allows general support for live internet streaming of game video and audio
       - Web cam video feed can be overlaid onto game viewports automatically
       - New 'Broadcast' Blueprint function library allows you to easily start broadcasting from your game
       - New IGameLiveStreaming API that allows you to start broadcasting through C++

- Editor live streaming
       - The editor UI now displays a "broadcast" button automatically when a live streaming service is available
       - Broadcasting of all desktop editor windows is supported (configured in preferences)
       - If you have a web cam, video will be displayed automatically while broadcasting in a new editor window
       - New "Live Streaming" editor preference tab with many new settings for configuring broadcasting
       - New IEditorLiveStreaming API that lets you control editor broadcasting directly, if needed

- Added new 'Broadcast.Start' and 'Broadcast.Stop' console commands
       - These allow you to easily test live streaming in games without writing UI code
       - Built-in help is available for these new commands

- To implement a live streaming plugin:
       - Inherit from the new ILiveStreamingService interface
       - Register your "LiveStreaming" feature with the IModularFeatures system

- Other changes:
      - Eliminated broken screen quad drawing functions; replaced with publicly-exposed DrawRectangle()
      - Slate: Eliminated legacy code for 'marking windows as drawn' (not used anymore)
      - Slate: Added Slate rendering callback to find out when a frame buffer is ready to be presented

diff --git a/Engine/Config/BaseEditorGameAgnostic.ini b/Engine/Config/BaseEditorGameAgnostic.ini
--- a/Engine/Config/BaseEditorGameAgnostic.ini
+++ b/Engine/Config/BaseEditorGameAgnostic.ini
@@ -6,3 +6,12 @@ LocalSourcePath="Editor/NewsFeed/[[Culture]]/"
 Source=NEWSFEED_Cdn
 MaxItemsToShow=10
 ShowOnlyUnreadItems=True
+
+[/Script/EditorLiveStreaming.EditorLiveStreamingSettings]
+FrameRate=30
+ScreenScaling=0.75
+bPrimaryMonitorOnly=true
+bEnableWebCam=true
+WebCamResolution=Normal_640x480
+bCaptureAudioFromComputer=true
+bCaptureAudioFromMicrophone=true
\ No newline at end of file
diff --git a/Engine/Source/Developer/CrashTracker/Private/CrashVideoCapture.cpp b/Engine/Source/Developer/CrashTracker/Private/CrashVideoCapture.cpp
--- a/Engine/Source/Developer/CrashTracker/Private/CrashVideoCapture.cpp
+++ b/Engine/Source/Developer/CrashTracker/Private/CrashVideoCapture.cpp
@@ -24,6 +24,10 @@ namespace CrashTrackerConstants
 
 	static const float SecondsToKeyPressFade = 2.f;
 	static const int32 MaxKeyPressesInBuffer = 15;
+
+	// Crash tracker video stream is scaled down by 50% to reduce the about of system resources used while
+	// working in the editor during capture
+	static const float ScreenScaling = 0.5f;
 }
 
 
@@ -294,7 +298,7 @@ void FCrashVideoCapture::Update(float DeltaSeconds)
 					AsyncTask->StartBackgroundTask();
 				}
 				
-				CaptureSlateRenderer->UnmapCrashTrackerBuffer();
+				CaptureSlateRenderer->UnmapVirtualScreenBuffer();
 			}
 			
 			++CurrentFrameCaptureIndex;
@@ -303,8 +307,12 @@ void FCrashVideoCapture::Update(float DeltaSeconds)
 				CurrentFrameCaptureIndex = 0;
 			}
 
-			CaptureSlateRenderer->CopyWindowsToDrawBuffer(StrippedKeyPressBuffer);
-			CaptureSlateRenderer->MapCrashTrackerBuffer(&Buffer[CurrentBufferIndex], &Width, &Height);
+			const bool bPrimaryWorkAreaOnly = false;	// We want all monitors captured for crash reporting
+			const FIntRect VirtualScreen = CaptureSlateRenderer->SetupVirtualScreenBuffer(bPrimaryWorkAreaOnly, CrashTrackerConstants::ScreenScaling, nullptr);
+			Width = VirtualScreen.Width();
+			Height = VirtualScreen.Height();
+			CaptureSlateRenderer->CopyWindowsToVirtualScreenBuffer(StrippedKeyPressBuffer);
+			CaptureSlateRenderer->MapVirtualScreenBuffer(&Buffer[CurrentBufferIndex]);
 
 			CurrentBufferIndex = (CurrentBufferIndex + 1) % 2;
 		}
diff --git a/Engine/Source/Editor/EditorLiveStreaming/EditorLiveStreaming.Build.cs b/Engine/Source/Editor/EditorLiveStreaming/EditorLiveStreaming.Build.cs
new file mode 100644
--- /dev/null
+++ b/Engine/Source/Editor/EditorLiveStreaming/EditorLiveStreaming.Build.cs
@@ -0,0 +1,45 @@
+// Copyright 1998-2014 Epic Games, Inc. All Rights Reserved.
+
+namespace UnrealBuildTool.Rules
+{
+	public class EditorLiveStreaming : ModuleRules
+	{
+		public EditorLiveStreaming( TargetInfo Target )
+		{
+			PublicDependencyModuleNames.AddRange(
+				new string[] 
+				{
+					"Core",
+				}
+			);
+
+			PrivateDependencyModuleNames.AddRange(
+				new string[] 
+				{
+					"CoreUObject",
+					"Engine",
+					"Slate",
+					"SlateCore",
+					"UnrealEd",
+				}
+			);
+
+			DynamicallyLoadedModuleNames.AddRange(
+				new string[] 
+				{
+					"Settings",
+					"MainFrame",
+					"GameLiveStreaming"
+				}
+			);
+
+			PrivateIncludePathModuleNames.AddRange(
+				new string[] 
+				{
+					"Settings",
+					"MainFrame"
+				}
+			);
+		}
+	}
+}
diff --git a/Engine/Source/Editor/EditorLiveStreaming/EditorLiveStreaming.cpp b/Engine/Source/Editor/EditorLiveStreaming/EditorLiveStreaming.cpp
new file mode 100644
--- /dev/null
+++ b/Engine/Source/Editor/EditorLiveStreaming/EditorLiveStreaming.cpp
@@ -0,0 +1,401 @@
+// Copyright 1998-2014 Epic Games, Inc. All Rights Reserved.
+#include "EditorLiveStreamingModule.h"
+#include "EditorLiveStreaming.h"
+#include "EditorLiveStreamingSettings.h"
+#include "Runtime/Core/Public/Features/IModularFeatures.h"
+#include "Runtime/Engine/Public/Features/ILiveStreamingService.h"
+#include "Runtime/SlateCore/Public/Rendering/SlateRenderer.h"
+#include "Developer/Settings/Public/Settings.h"
+#include "Editor/MainFrame/Public/Interfaces/IMainFrameModule.h"
+#include "Runtime/GameLiveStreaming/Public/IGameLiveStreaming.h"
+
+IMPLEMENT_MODULE( FEditorLiveStreaming, "EditorLiveStreaming" );
+
+#define LOCTEXT_NAMESPACE "EditorLiveStreaming"
+
+FEditorLiveStreaming::FEditorLiveStreaming()
+	: bIsBroadcasting( false ),
+	  LiveStreamer( nullptr ),
+	  SubmittedVideoFrameCount( 0 ),
+	  ReadbackBufferIndex( 0 )
+{
+	ReadbackBuffers[0] = ReadbackBuffers[1] = nullptr;
+}
+
+
+void FEditorLiveStreaming::StartupModule()
+{
+	// Register settings
+	ISettingsModule* SettingsModule = ISettingsModule::Get();
+	if( ensure( SettingsModule != nullptr ) )
+	{
+		SettingsModule->RegisterSettings( "Editor", "General", "EditorLiveStreaming",
+			LOCTEXT( "EditorLiveStreamingSettingsName", "Live Streaming"),
+			LOCTEXT( "EditorLiveStreamingSettingsDescription", "Settings for broadcasting live from the editor"),
+			GetMutableDefault< UEditorLiveStreamingSettings >()
+		);
+	}
+}
+
+
+void FEditorLiveStreaming::ShutdownModule()
+{
+	// Unregister settings
+	ISettingsModule* SettingsModule = ISettingsModule::Get();
+	if( SettingsModule != nullptr )
+	{
+		SettingsModule->UnregisterSettings( "Editor", "General", "EditorLiveStreaming" );
+	}
+}
+
+
+bool FEditorLiveStreaming::IsLiveStreamingAvailable() const
+{
+	bool bIsAvailable = false;
+
+	// We are not available if the game is already streaming
+	if( !IGameLiveStreaming::Get().IsBroadcastingGame() )
+	{
+		static const FName LiveStreamingFeatureName( "LiveStreaming" );
+		if( IModularFeatures::Get().IsModularFeatureAvailable( LiveStreamingFeatureName ) )
+		{
+			// Only if not already broadcasting
+			bIsAvailable = true;
+		}
+	}
+	return bIsAvailable;
+}
+
+
+bool FEditorLiveStreaming::IsBroadcastingEditor() const
+{
+	return bIsBroadcasting;
+}
+
+
+void FEditorLiveStreaming::StartBroadcastingEditor()
+{
+	if( !IsBroadcastingEditor() && IsLiveStreamingAvailable() )
+	{
+		// Select a live streaming service
+		{
+			static const FName LiveStreamingFeatureName( "LiveStreaming" );
+
+			// @todo livestream: Always using first live streaming plugin we find, but instead this should be a configurable option in the
+			// editor's preferences, based on the enumeration of available features from plugins
+			LiveStreamer = static_cast<ILiveStreamingService*>( IModularFeatures::Get().GetModularFeatureImplementation( LiveStreamingFeatureName, 0 ) );
+		}
+
+		// Register to find out about status changes
+		LiveStreamer->OnStatusChanged().AddRaw( this, &FEditorLiveStreaming::BroadcastStatusCallback );
+
+
+		// Tell our live streaming plugin to start broadcasting
+		{
+			const auto& Settings = *GetDefault< UEditorLiveStreamingSettings >();
+			FSlateRenderer* SlateRenderer = FSlateApplication::Get().GetRenderer().Get();
+			const FIntRect VirtualScreen = SlateRenderer->SetupVirtualScreenBuffer( Settings.bPrimaryMonitorOnly, Settings.ScreenScaling, LiveStreamer );
+
+			bIsBroadcasting = true;
+			SubmittedVideoFrameCount = 0;
+
+			// @todo livestream: What about if virtual screen size changes while we are still broadcasting?  For example, if the user changes their
+			// desktop resolution while the editor is running.  We'd need to stop and restart the broadcast.
+			FBroadcastConfig BroadcastConfig;
+			BroadcastConfig.VideoBufferWidth = VirtualScreen.Width();
+			BroadcastConfig.VideoBufferHeight = VirtualScreen.Height();
+			BroadcastConfig.FramesPerSecond = Settings.FrameRate;
+			BroadcastConfig.PixelFormat = FBroadcastConfig::EBroadcastPixelFormat::R8G8B8A8;
+			BroadcastConfig.bCaptureAudioFromComputer = Settings.bCaptureAudioFromComputer;
+			BroadcastConfig.bCaptureAudioFromMicrophone = Settings.bCaptureAudioFromMicrophone;
+			LiveStreamer->StartBroadcasting( BroadcastConfig );
+
+			if( Settings.bEnableWebCam )
+			{
+				FWebCamConfig WebCamConfig;
+				switch( Settings.WebCamResolution )
+				{
+					case EEditorLiveStreamingWebCamResolution::Normal_320x240:
+						WebCamConfig.DesiredWebCamWidth = 320;
+						WebCamConfig.DesiredWebCamHeight = 240;
+						break;
+					case EEditorLiveStreamingWebCamResolution::Wide_320x180:
+						WebCamConfig.DesiredWebCamWidth = 320;
+						WebCamConfig.DesiredWebCamHeight = 180;
+						break;
+					case EEditorLiveStreamingWebCamResolution::Normal_640x480:
+						WebCamConfig.DesiredWebCamWidth = 640;
+						WebCamConfig.DesiredWebCamHeight = 480;
+						break;
+					case EEditorLiveStreamingWebCamResolution::Wide_640x360:
+						WebCamConfig.DesiredWebCamWidth = 640;
+						WebCamConfig.DesiredWebCamHeight = 360;
+						break;
+					case EEditorLiveStreamingWebCamResolution::Normal_800x600:
+						WebCamConfig.DesiredWebCamWidth = 800;
+						WebCamConfig.DesiredWebCamHeight = 600;
+						break;
+					case EEditorLiveStreamingWebCamResolution::Wide_800x450:
+						WebCamConfig.DesiredWebCamWidth = 800;
+						WebCamConfig.DesiredWebCamHeight = 450;
+						break;
+					case EEditorLiveStreamingWebCamResolution::Normal_1024x768:
+						WebCamConfig.DesiredWebCamWidth = 1024;
+						WebCamConfig.DesiredWebCamHeight = 768;
+						break;
+					case EEditorLiveStreamingWebCamResolution::Wide_1024x576:
+						WebCamConfig.DesiredWebCamWidth = 1024;
+						WebCamConfig.DesiredWebCamHeight = 576;
+						break;
+					case EEditorLiveStreamingWebCamResolution::Normal_1080x810:
+						WebCamConfig.DesiredWebCamWidth = 1080;
+						WebCamConfig.DesiredWebCamHeight = 810;
+						break;
+					case EEditorLiveStreamingWebCamResolution::Wide_1080x720:
+						WebCamConfig.DesiredWebCamWidth = 1080;
+						WebCamConfig.DesiredWebCamHeight = 720;
+						break;
+					case EEditorLiveStreamingWebCamResolution::Normal_1280x960:
+						WebCamConfig.DesiredWebCamWidth = 1280;
+						WebCamConfig.DesiredWebCamHeight = 960;
+						break;
+					case EEditorLiveStreamingWebCamResolution::Wide_1280x720:
+						WebCamConfig.DesiredWebCamWidth = 1280;
+						WebCamConfig.DesiredWebCamHeight = 720;
+						break;
+					case EEditorLiveStreamingWebCamResolution::Normal_1920x1440:
+						WebCamConfig.DesiredWebCamWidth = 1920;
+						WebCamConfig.DesiredWebCamHeight = 1440;
+						break;
+					case EEditorLiveStreamingWebCamResolution::Wide_1920x1080:
+						WebCamConfig.DesiredWebCamWidth = 1920;
+						WebCamConfig.DesiredWebCamHeight = 1080;
+						break;
+
+					default:
+						check(0);
+						break;
+				}
+
+				// @todo livestream: Allow web cam to be started/stopped independently from the broadcast itself, so users can setup their web cam
+				LiveStreamer->StartWebCam( WebCamConfig );
+			}
+		}
+	}
+}
+
+
+void FEditorLiveStreaming::StopBroadcastingEditor()
+{
+	if( IsBroadcastingEditor() )
+	{
+		CloseBroadcastStatusWindow();
+
+		LiveStreamer->StopWebCam();
+		LiveStreamer->StopBroadcasting();
+
+		// Unregister for status changes
+		LiveStreamer->OnStatusChanged().RemoveAll( this );
+
+		LiveStreamer = nullptr;
+
+		// Tell our notification to start expiring
+		TSharedPtr< SNotificationItem > Notification( NotificationWeakPtr.Pin() );
+		if( Notification.IsValid() )
+		{
+			Notification->ExpireAndFadeout();
+		}
+
+		bIsBroadcasting = false;
+		for( int32 BufferIndex = 0; BufferIndex < 2; ++BufferIndex )
+		{
+			ReadbackBuffers[ BufferIndex ] = nullptr;
+		}
+		ReadbackBufferIndex = 0;
+		SubmittedVideoFrameCount = 0;
+	}
+}
+
+
+void FEditorLiveStreaming::BroadcastEditorVideoFrame()
+{
+	if( ensure( IsBroadcastingEditor() ) )
+	{
+		// Check to see if we're streaming live video.  If so, we'll want to push new frames to be broadcast.
+		if( LiveStreamer->IsBroadcasting() )
+		{
+			// Is this live streamer ready to accept new frames?
+			if( LiveStreamer->IsReadyForVideoFrames() )
+			{
+				FSlateRenderer* SlateRenderer = FSlateApplication::Get().GetRenderer().Get();
+
+				if( ReadbackBuffers[ ReadbackBufferIndex ] != nullptr )
+				{
+ 					LiveStreamer->PushVideoFrame( (FColor*)ReadbackBuffers[ ReadbackBufferIndex ] );
+					++SubmittedVideoFrameCount;
+
+					// If this is the first frame we've submitted, then we can fade out our notification UI, since broadcasting is in progress
+					if( SubmittedVideoFrameCount == 1 )
+					{
+						TSharedPtr< SNotificationItem > Notification( NotificationWeakPtr.Pin() );
+						if( Notification.IsValid() )
+						{
+							Notification->ExpireAndFadeout();
+						}
+					}
+
+					SlateRenderer->UnmapVirtualScreenBuffer();
+				}
+
+				TArray<FString> UnusedKeypressBuffer;
+				SlateRenderer->CopyWindowsToVirtualScreenBuffer( UnusedKeypressBuffer );
+				SlateRenderer->MapVirtualScreenBuffer( &ReadbackBuffers[ ReadbackBufferIndex ] );
+
+				// Ping pong between buffers
+				ReadbackBufferIndex = ( ReadbackBufferIndex + 1 ) % 2;
+			}
+		}
+		else
+		{
+			// If our streaming service has stopped broadcasting for some reason, then we'll cancel our broadcast
+			StopBroadcastingEditor();
+		}
+	}
+}
+
+
+void FEditorLiveStreaming::BroadcastStatusCallback( const FLiveStreamingStatus& Status )
+{
+	TSharedPtr< SNotificationItem > Notification( NotificationWeakPtr.Pin() );
+	if( Notification.IsValid() )
+	{
+		// Don't bother clobbering existing message with text about web cam starting/stopping, unless we're already broadcasting.
+		// We want to make sure people see the persistent text about login state.
+		if( SubmittedVideoFrameCount > 0 ||
+			( Status.StatusType != FLiveStreamingStatus::EStatusType::WebCamStarted &&
+			  Status.StatusType != FLiveStreamingStatus::EStatusType::WebCamStopped &&
+			  Status.StatusType != FLiveStreamingStatus::EStatusType::WebCamTextureNotReady &&
+			  Status.StatusType != FLiveStreamingStatus::EStatusType::WebCamTextureReady ) )
+		{
+			Notification->SetText( Status.CustomStatusDescription );
+		}
+	}
+	else
+	{
+		// Only spawn a notification if we're actually still trying to broadcast, not if we're stopping broadcasting.  We don't want
+		// to revive our notification that we intentionally expired.
+		if( bIsBroadcasting )
+		{
+			FNotificationInfo Info( Status.CustomStatusDescription );
+			Info.FadeInDuration = 0.1f;
+			Info.FadeOutDuration = 0.5f;
+			Info.ExpireDuration = 1.5f;
+			Info.bUseThrobber = false;
+			Info.bUseSuccessFailIcons = true;
+			Info.bUseLargeFont = true;
+			Info.bFireAndForget = false;
+			Info.bAllowThrottleWhenFrameRateIsLow = false;
+			NotificationWeakPtr = FSlateNotificationManager::Get().AddNotification( Info );
+		}
+	}
+
+	Notification = NotificationWeakPtr.Pin();
+	if( Notification.IsValid() )
+	{
+		SNotificationItem::ECompletionState State = SNotificationItem::CS_Pending;
+		if( Status.StatusType == FLiveStreamingStatus::EStatusType::Failure )
+		{
+			State = SNotificationItem::CS_Fail;
+		}
+		else if( Status.StatusType == FLiveStreamingStatus::EStatusType::BroadcastStarted ||
+				 Status.StatusType == FLiveStreamingStatus::EStatusType::WebCamStarted )
+		{
+			State = SNotificationItem::CS_Success;
+		}
+
+		Notification->SetCompletionState( State );
+	}
+
+	// If the web cam just turned on, then we'll go ahead and show it
+	if( Status.StatusType == FLiveStreamingStatus::EStatusType::WebCamTextureReady )
+	{	
+		UTexture2D* WebCamTexture = LiveStreamer->GetWebCamTexture();
+		if( ensure( WebCamTexture != nullptr ) )
+		{
+			IMainFrameModule& MainFrameModule = FModuleManager::LoadModuleChecked<IMainFrameModule>( TEXT( "MainFrame" ) );
+			check( MainFrameModule.IsWindowInitialized() );
+			auto RootWindow = MainFrameModule.GetParentWindow();
+			check( RootWindow.IsValid() );
+
+			// How many pixels from the edge of the main frame window to where the broadcast status window appears
+			const int WindowBorderPadding = 50;
+
+			// @todo livestream: Currently this window is not created as "topmost".  We don't really want it to be OS-topmost, but we do want it to be the most
+			// topmost "regular" window in our application (not on top of tooltips, etc.)
+
+			// Create a window that will show the web cam video feed
+			BroadcastStatusWindow = SNew( SWindow )
+				.Title( LOCTEXT( "StreamingWindowTitle", "Web Camera" ) )
+				.ClientSize( FVector2D( WebCamTexture->GetSizeX(), WebCamTexture->GetSizeY() ) )
+				.ScreenPosition( FVector2D( 
+					RootWindow->GetRectInScreen().Right - WebCamTexture->GetSizeX() - WindowBorderPadding, 
+					RootWindow->GetRectInScreen().Top +  WindowBorderPadding ) )
+
+				// @todo livestream: Ideally the user could freely resize the window, but it gets a bit tricky to preserve the web cam aspect.  Plus, we don't really like how letterboxing/columnboxing looks with this.
+				.SizingRule( ESizingRule::FixedSize )
+
+				.AutoCenter( EAutoCenter::None )
+				.bDragAnywhere( true )
+				.SupportsMaximize( true )
+				.SupportsMinimize( true )
+				.FocusWhenFirstShown( false )
+				.ActivateWhenFirstShown( false )
+				.SaneWindowPlacement( false );
+
+			WebCamDynamicImageBrush = MakeShareable( new FSlateDynamicImageBrush( 
+				WebCamTexture,
+				FVector2D( WebCamTexture->GetSizeX(), WebCamTexture->GetSizeY() ),
+				WebCamTexture->GetFName() ) );
+
+			// @todo livestream: Currently if the user closes the window, the camera is deactivated and it doesn't turn back on unless the broadcast is restarted.  We could allow the camera to be reactivated though.
+			BroadcastStatusWindow->SetOnWindowClosed( 
+				FOnWindowClosed::CreateStatic( []( const TSharedRef<SWindow>& Window, FEditorLiveStreaming* This ) 
+				{
+					// User closed the broadcast status window, so go ahead and shut down the web cam
+					This->LiveStreamer->StopWebCam();
+					This->BroadcastStatusWindow.Reset();
+					This->WebCamDynamicImageBrush.Reset();
+				}, 
+				this ) );
+
+			BroadcastStatusWindow->SetContent( 
+				SNew( SImage )
+					.Image( WebCamDynamicImageBrush.Get() )
+			);
+
+			FSlateApplication::Get().AddWindowAsNativeChild(
+				BroadcastStatusWindow.ToSharedRef(), 
+				RootWindow.ToSharedRef() );
+		}
+
+	}
+	else if( Status.StatusType == FLiveStreamingStatus::EStatusType::WebCamTextureNotReady )
+	{
+		CloseBroadcastStatusWindow();
+	}
+}
+
+
+void FEditorLiveStreaming::CloseBroadcastStatusWindow()
+{
+	if( BroadcastStatusWindow.IsValid() )
+	{
+		FSlateApplication::Get().DestroyWindowImmediately( BroadcastStatusWindow.ToSharedRef() );
+		BroadcastStatusWindow.Reset();
+		WebCamDynamicImageBrush.Reset();
+	}
+}
+
+
+
+#undef LOCTEXT_NAMESPACE
diff --git a/Engine/Source/Editor/EditorLiveStreaming/EditorLiveStreaming.h b/Engine/Source/Editor/EditorLiveStreaming/EditorLiveStreaming.h
new file mode 100644
--- /dev/null
+++ b/Engine/Source/Editor/EditorLiveStreaming/EditorLiveStreaming.h
@@ -0,0 +1,64 @@
+// Copyright 1998-2014 Epic Games, Inc. All Rights Reserved.
+#pragma once
+
+#include "IEditorLiveStreaming.h"
+
+
+/**
+ * Handles broadcasting of the editor to a live internet stream
+ */
+class FEditorLiveStreaming
+	: public IEditorLiveStreaming
+{
+
+public:
+
+	/** Default constructor */
+	FEditorLiveStreaming();
+
+	/** IModuleInterface overrides */
+	virtual void StartupModule() override;
+	virtual void ShutdownModule() override;
+
+	/** IEditorLiveStreaming overrides */
+	virtual bool IsLiveStreamingAvailable() const override;
+	virtual bool IsBroadcastingEditor() const override;
+	virtual void StartBroadcastingEditor() override;
+	virtual void StopBroadcastingEditor() override;
+	virtual void BroadcastEditorVideoFrame() override;
+
+protected:
+
+	/** Called by the live streaming service when streaming status has changed or an error has occurred */
+	void BroadcastStatusCallback( const struct FLiveStreamingStatus& Status );
+
+	/** Closes the window that shows the current status of the broadcast */
+	void CloseBroadcastStatusWindow();
+
+
+private:
+
+	/** True if we're currently broadcasting anything */
+	bool bIsBroadcasting;
+
+	/** The live streaming implementation we're using for editor broadcasting */
+	class ILiveStreamingService* LiveStreamer;
+
+	/** Number of video frames we've submitted since we started broadcasting */
+	uint32 SubmittedVideoFrameCount;
+
+	/** A temporary pair of buffers that we write our video frames to */
+	void* ReadbackBuffers[2];
+
+	/** The current buffer index.  We bounce between them to avoid stalls. */
+	int32 ReadbackBufferIndex;
+
+	/** Persistent Slate notification for live streaming status */
+	TWeakPtr< class SNotificationItem > NotificationWeakPtr;
+
+	/** Dynamic slate image brush for our web cam texture */
+	TSharedPtr< struct FSlateDynamicImageBrush > WebCamDynamicImageBrush;
+
+	/** Window for broadcast status */
+	TSharedPtr< class SWindow > BroadcastStatusWindow;
+};
diff --git a/Engine/Source/Editor/EditorLiveStreaming/EditorLiveStreamingModule.h b/Engine/Source/Editor/EditorLiveStreaming/EditorLiveStreamingModule.h
new file mode 100644
--- /dev/null
+++ b/Engine/Source/Editor/EditorLiveStreaming/EditorLiveStreamingModule.h
@@ -0,0 +1,5 @@
+// Copyright 1998-2014 Epic Games, Inc. All Rights Reserved.
+#pragma once
+
+#include "UnrealEd.h"
+
diff --git a/Engine/Source/Editor/EditorLiveStreaming/EditorLiveStreamingSettings.cpp b/Engine/Source/Editor/EditorLiveStreaming/EditorLiveStreamingSettings.cpp
new file mode 100644
--- /dev/null
+++ b/Engine/Source/Editor/EditorLiveStreaming/EditorLiveStreamingSettings.cpp
@@ -0,0 +1,27 @@
+// Copyright 1998-2014 Epic Games, Inc. All Rights Reserved.
+#include "EditorLiveStreamingModule.h"
+#include "EditorLiveStreamingSettings.h"
+
+
+UEditorLiveStreamingSettings::UEditorLiveStreamingSettings( const class FPostConstructInitializeProperties& PCIP )
+	: Super( PCIP )
+{
+}
+
+
+void UEditorLiveStreamingSettings::PostEditChangeProperty( struct FPropertyChangedEvent& PropertyChangedEvent )
+{
+	UProperty* PropertyThatChanged = PropertyChangedEvent.Property;
+
+	const FName Name = PropertyThatChanged ? PropertyThatChanged->GetFName() : NAME_None;
+// 	if( Name == ??? )
+// 	{
+//		...
+//	}
+
+	// Save config to file, but only if we are not the build machine since game agnostic settings may put the builder in an unclean state
+	if( !GIsBuildMachine )
+	{
+		this->SaveConfig();
+	}	
+}
\ No newline at end of file
diff --git a/Engine/Source/Editor/EditorLiveStreaming/EditorLiveStreamingSettings.h b/Engine/Source/Editor/EditorLiveStreaming/EditorLiveStreamingSettings.h
new file mode 100644
--- /dev/null
+++ b/Engine/Source/Editor/EditorLiveStreaming/EditorLiveStreamingSettings.h
@@ -0,0 +1,80 @@
+// Copyright 1998-2014 Epic Games, Inc. All Rights Reserved.
+
+#pragma once
+#include "EditorLiveStreamingSettings.generated.h"
+
+
+/** Describes the available preset resolutions for web camera video with the editor's live streaming feature */
+UENUM()
+namespace EEditorLiveStreamingWebCamResolution
+{
+	enum Type
+	{
+		Normal_320x240 UMETA( DisplayName="320 x 240 (4:3)" ),
+		Wide_320x180 UMETA( DisplayName="320 x 180 (16:9)" ),
+
+		Normal_640x480 UMETA( DisplayName="640 x 480 (4:3)" ),
+		Wide_640x360 UMETA( DisplayName="640 x 360 (16:9)" ),
+
+		Normal_800x600 UMETA( DisplayName="800 x 600 (4:3)" ),
+		Wide_800x450 UMETA( DisplayName="800 x 450 (16:9)" ),
+
+		Normal_1024x768 UMETA( DisplayName="1024 x 768 (4:3)" ),
+		Wide_1024x576 UMETA( DisplayName="1024 x 576 (16:9)" ),
+		
+		Normal_1080x810 UMETA( DisplayName="1080 x 810 (4:3)" ),
+		Wide_1080x720 UMETA( DisplayName="1080 x 720 (16:9)" ),
+
+		Normal_1280x960 UMETA( DisplayName="1280 x 960 (4:3)" ),
+		Wide_1280x720 UMETA( DisplayName="1280 x 720 (16:9, 720p)" ),
+
+		Normal_1920x1440 UMETA( DisplayName="1920 x 1440 (4:3)" ),
+		Wide_1920x1080 UMETA( DisplayName="1920 x 1080 (16:9, 1080p)" ),
+	};
+}
+
+
+/** Holds preferences for the editor live streaming features */
+UCLASS(config=EditorGameAgnostic)
+class UEditorLiveStreamingSettings : public UObject
+{
+	GENERATED_UCLASS_BODY()
+
+public:
+
+	/** Frame rate to stream video from the editor at when broadcasting to services like Twitch. */
+	UPROPERTY(Config, EditAnywhere, Category=EditorLiveStreaming, Meta=(UIMin=5,ClampMin=1,UIMax=60))
+	int32 FrameRate;
+
+	/** How much to scale the broadcast video resolution down to reduce streaming bandwidth.  We recommend broadcasting at resolutions of 1280x720 or lower.  Some live streaming providers will not be able to transcode your video to a lower resolution, so using a high resolution stream may prevent low-bandwidth users from having a good viewing experience.*/
+	UPROPERTY(Config, EditAnywhere, Category=EditorLiveStreaming, Meta=(ClampMin=0.1,ClampMax=1.0))
+	float ScreenScaling;
+
+	/** By default, we only broadcast video from your primary monitor's work area (excludes the task bar area.)  Turning off this feature enables broadcasting from all monitors, covering your entire usable desktop area.  This may greatly increase the video resolution of your stream, so we recommend leaving this option turned off. */
+	UPROPERTY(Config, EditAnywhere, Category=EditorLiveStreaming )
+	bool bPrimaryMonitorOnly;
+
+	/** If enabled, video from your web camera will be captured and displayed in the editor while broadcasting, so that your viewers can see your presence. */
+	UPROPERTY(Config, EditAnywhere, Category=EditorLiveStreaming )
+	bool bEnableWebCam;
+
+	/** The camera video resolution that you'd prefer.  The camera may not support the exact resolution you specify, so we'll try to find the best match automatically */
+	UPROPERTY(Config, EditAnywhere, Category=EditorLiveStreaming, Meta=(EditCondition="bEnableWebCam") )
+	TEnumAsByte<EEditorLiveStreamingWebCamResolution::Type> WebCamResolution;
+
+	/** Enables broadcast of audio being played by your computer, such as in-game sounds */
+	UPROPERTY(Config, EditAnywhere, Category=EditorLiveStreaming )
+	bool bCaptureAudioFromComputer;
+
+	/** Enables broadcast of audio from your default microphone recording device */
+	UPROPERTY(Config, EditAnywhere, Category=EditorLiveStreaming )
+	bool bCaptureAudioFromMicrophone;
+
+
+protected:
+
+	/** UObject overrides */
+	virtual void PostEditChangeProperty( struct FPropertyChangedEvent& PropertyChangedEvent) override;
+};
+
+
diff --git a/Engine/Source/Editor/EditorLiveStreaming/Public/IEditorLiveStreaming.h b/Engine/Source/Editor/EditorLiveStreaming/Public/IEditorLiveStreaming.h
new file mode 100644
--- /dev/null
+++ b/Engine/Source/Editor/EditorLiveStreaming/Public/IEditorLiveStreaming.h
@@ -0,0 +1,54 @@
+// Copyright 1998-2014 Epic Games, Inc. All Rights Reserved.
+#pragma once
+
+#include "ModuleManager.h"
+
+
+/**
+ * Interface to the EditorLiveStreaming module's functionality that enables live broadcasting of video and audio from editor sessions
+ */
+class IEditorLiveStreaming
+	: public IModuleInterface
+{
+
+public:
+
+	/**
+	 * Checks to see if live streaming is currently available for the editor 
+	 *
+	 * @return	True if live streaming is possible to start
+	 */
+	virtual bool IsLiveStreamingAvailable() const = 0;
+
+	/**
+	 * Returns true if we're currently broadcasting
+	 *
+	 * @return	True if broadcasting is currently active
+	 */
+	virtual bool IsBroadcastingEditor() const = 0;
+
+	/** Starts broadcasting frames */
+	virtual void StartBroadcastingEditor() = 0;
+
+	/** Stops broadcasting frames */
+	virtual void StopBroadcastingEditor() = 0;
+
+	/** Pushes latest video from the editor to a live stream.  The editor's rendering system will call this each frame. */
+	virtual void BroadcastEditorVideoFrame() = 0;
+
+
+public:
+
+	/**
+	 * Singleton-like access to IEditorLiveStreaming
+	 *
+	 * @return Returns instance of IEditorLiveStreaming object
+	 */
+	static inline IEditorLiveStreaming& Get()
+	{
+		return FModuleManager::LoadModuleChecked< IEditorLiveStreaming >( "EditorLiveStreaming" );
+	}
+
+};
+
+
diff --git a/Engine/Source/Editor/EditorStyle/Private/SlateEditorStyle.cpp b/Engine/Source/Editor/EditorStyle/Private/SlateEditorStyle.cpp
--- a/Engine/Source/Editor/EditorStyle/Private/SlateEditorStyle.cpp
+++ b/Engine/Source/Editor/EditorStyle/Private/SlateEditorStyle.cpp
@@ -1365,6 +1365,11 @@ void FSlateEditorStyle::FStyle::SetupGeneralStyles()
 	}
 
 #if WITH_EDITOR || IS_PROGRAM
+	// Editor live streaming
+	{
+		Set( "EditorLiveStreaming.BroadcastButton", new IMAGE_BRUSH( "Animation/Record", Icon20x20 ) );
+	}
+
 	// Collision Analyzer
 	{
 		Set( "CollisionAnalyzer.TabIcon", new IMAGE_BRUSH( "Icons/icon_tab_CollisionAnalyser_16x", Icon16x16 ) );
diff --git a/Engine/Source/Editor/MainFrame/MainFrame.Build.cs b/Engine/Source/Editor/MainFrame/MainFrame.Build.cs
--- a/Engine/Source/Editor/MainFrame/MainFrame.Build.cs
+++ b/Engine/Source/Editor/MainFrame/MainFrame.Build.cs
@@ -72,6 +72,7 @@ public class MainFrame : ModuleRules
 				"OutputLog",
                 "TranslationEditor",
 				"SourceCodeAccess",
+				"EditorLiveStreaming"
 			}
 		);
 	}
diff --git a/Engine/Source/Editor/MainFrame/Private/MainFrameModule.cpp b/Engine/Source/Editor/MainFrame/Private/MainFrameModule.cpp
--- a/Engine/Source/Editor/MainFrame/Private/MainFrameModule.cpp
+++ b/Engine/Source/Editor/MainFrame/Private/MainFrameModule.cpp
@@ -2,7 +2,7 @@
 
 #include "MainFramePrivatePCH.h"
 #include "CompilerResultsLog.h"
-
+#include "Editor/EditorLiveStreaming/Public/IEditorLiveStreaming.h"
 
 DEFINE_LOG_CATEGORY(LogMainFrame);
 #define LOCTEXT_NAMESPACE "FMainFrameModule"
@@ -498,7 +498,46 @@ TSharedRef<SWidget> FMainFrameModule::MakeDeveloperTools() const
 					ISourceControlModule::Get().CreateStatusWidget().ToSharedRef()
 				]
 
-			//save video button
+			// Editor live streaming toggle button
+			+SHorizontalBox::Slot()
+				.AutoWidth()
+				.VAlign( VAlign_Bottom )
+				[
+					SNew(SButton)
+					.Visibility_Static( []() -> EVisibility { return IEditorLiveStreaming::Get().IsLiveStreamingAvailable() ? EVisibility::Visible : EVisibility::Collapsed; } )
+					.ToolTipText( LOCTEXT( "BroadcastTooltip", "Starts or stops broadcasting of this editor session to a live internet streaming service." ) )
+					.OnClicked_Static( []
+						{ 
+							// Toggle broadcasting on or off
+							if( IEditorLiveStreaming::Get().IsBroadcastingEditor() ) 
+							{
+								IEditorLiveStreaming::Get().StopBroadcastingEditor();
+							}
+							else
+							{
+								IEditorLiveStreaming::Get().StartBroadcastingEditor();
+							}
+							return FReply::Handled();
+						} )
+					.ButtonStyle( FEditorStyle::Get(), "NoBorder" )
+					.ContentPadding(FMargin(1,0))
+					[
+						SNew(SImage)
+						.Image( FEditorStyle::GetBrush("EditorLiveStreaming.BroadcastButton") )
+						.ColorAndOpacity_Static( [] 
+							{ 
+								// Pulsate the button graphics while we're broadcasting
+								FSlateColor Color( FLinearColor::White );
+								if( IEditorLiveStreaming::Get().IsBroadcastingEditor() )
+								{
+									Color = FLinearColor( 1.0f, 1.0f, 1.0f, FMath::MakePulsatingValue( FSlateApplication::Get().GetCurrentTime(), 2.0f ) );
+								}
+								return Color;
+							} )
+					]
+				]
+
+			// Crash report "save video" button
 			+SHorizontalBox::Slot()
 				.AutoWidth()
 				.VAlign( VAlign_Bottom )
diff --git a/Engine/Source/Editor/UnrealEd/Private/Editor.cpp b/Engine/Source/Editor/UnrealEd/Private/Editor.cpp
--- a/Engine/Source/Editor/UnrealEd/Private/Editor.cpp
+++ b/Engine/Source/Editor/UnrealEd/Private/Editor.cpp
@@ -541,6 +541,7 @@ void UEditorEngine::Init(IEngineLoop* InEngineLoop)
 		FModuleManager::Get().LoadModule(TEXT("UndoHistory"));
 		FModuleManager::Get().LoadModule(TEXT("DeviceProfileEditor"));
 		FModuleManager::Get().LoadModule(TEXT("SourceCodeAccess"));
+		FModuleManager::Get().LoadModule(TEXT("EditorLiveStreaming"));
 
 		if ( FParse::Param(FCommandLine::Get(), TEXT("umg")) )
 		{
diff --git a/Engine/Source/Editor/UnrealEd/Private/UnrealEdEngine.cpp b/Engine/Source/Editor/UnrealEd/Private/UnrealEdEngine.cpp
--- a/Engine/Source/Editor/UnrealEd/Private/UnrealEdEngine.cpp
+++ b/Engine/Source/Editor/UnrealEd/Private/UnrealEdEngine.cpp
@@ -26,7 +26,7 @@
 #include "PackageAutoSaver.h"
 #include "BSPOps.h"
 #include "ComponentVisualizer.h"
-
+#include "Editor/EditorLiveStreaming/Public/IEditorLiveStreaming.h"
 
 DEFINE_LOG_CATEGORY_STATIC(LogUnrealEdEngine, Log, All);
 
@@ -283,9 +283,22 @@ void UUnrealEdEngine::Tick(float DeltaSeconds, bool bIdleMode)
 	}
 	
 	ICrashTrackerModule* CrashTracker = FModuleManager::LoadModulePtr<ICrashTrackerModule>( FName("CrashTracker") );
+	bool bCrashTrackerEnabled = false;
 	if (CrashTracker)
 	{
 		CrashTracker->Update(DeltaSeconds);
+		bCrashTrackerEnabled = CrashTracker->IsCurrentlyCapturing();
+	}
+
+	// Only allow live streaming if crash tracker is disabled. This is because the SlateRHIRenderer shares the same render targets
+	// for both crash tracker and live editor streaming, and we don't want them to be thrashed every frame.
+	if( !bCrashTrackerEnabled )
+	{
+		// If the editor is configured to broadcast frames, do that now
+		if( IEditorLiveStreaming::Get().IsBroadcastingEditor() )
+		{
+			IEditorLiveStreaming::Get().BroadcastEditorVideoFrame();
+		}
 	}
 }
 
diff --git a/Engine/Source/Editor/UnrealEd/UnrealEd.Build.cs b/Engine/Source/Editor/UnrealEd/UnrealEd.Build.cs
--- a/Engine/Source/Editor/UnrealEd/UnrealEd.Build.cs
+++ b/Engine/Source/Editor/UnrealEd/UnrealEd.Build.cs
@@ -173,7 +173,8 @@ public class UnrealEd : ModuleRules
                 "GameplayAbilitiesEditor",
 				"UndoHistory",
 				"SourceCodeAccess",
-				"ReferenceViewer"
+				"ReferenceViewer",
+				"EditorLiveStreaming"
 			}
 		);
 
diff --git a/Engine/Source/Runtime/Engine/Engine.Build.cs b/Engine/Source/Runtime/Engine/Engine.Build.cs
--- a/Engine/Source/Runtime/Engine/Engine.Build.cs
+++ b/Engine/Source/Runtime/Engine/Engine.Build.cs
@@ -142,6 +142,7 @@ public class Engine : ModuleRules
 			    new string[] {
 				    "CrashTracker",
 				    "ImageWrapper",
+					"GameLiveStreaming"
 			    }
 		    );
         }
diff --git a/Engine/Source/Runtime/Engine/Private/GameEngine.cpp b/Engine/Source/Runtime/Engine/Private/GameEngine.cpp
--- a/Engine/Source/Runtime/Engine/Private/GameEngine.cpp
+++ b/Engine/Source/Runtime/Engine/Private/GameEngine.cpp
@@ -410,6 +410,14 @@ void UGameEngine::Init(IEngineLoop* InEngineLoop)
 	}
 #endif
 
+	// Load game modules
+	{
+		if( !IsRunningDedicatedServer() && !IsRunningCommandlet() )
+		{
+			FModuleManager::Get().LoadModule( TEXT("GameLiveStreaming") );
+		}
+	}
+
 	// Load and apply user game settings
 	GetGameUserSettings()->LoadSettings();
 	GetGameUserSettings()->ApplySettings();
diff --git a/Engine/Source/Runtime/Engine/Private/GameViewportClient.cpp b/Engine/Source/Runtime/Engine/Private/GameViewportClient.cpp
--- a/Engine/Source/Runtime/Engine/Private/GameViewportClient.cpp
+++ b/Engine/Source/Runtime/Engine/Private/GameViewportClient.cpp
@@ -17,10 +17,9 @@
 #include "EngineModule.h"
 #include "AudioDevice.h"
 #include "Sound/SoundWave.h"
-
 #include "HighResScreenshot.h"
-
 #include "Particles/ParticleSystemComponent.h"
+#include "Runtime/GameLiveStreaming/Public/IGameLiveStreaming.h"
 
 /** This variable allows forcing full screen of the first player controller viewport, even if there are multiple controllers plugged in and no cinematic playing. */
 bool GForceFullscreen = false;
@@ -1180,6 +1179,10 @@ void UGameViewportClient::PostRender(UCanvas* Canvas)
 
 	// Draw the transition screen.
 	DrawTransition(Canvas);
+
+	// Draw default web cam.  This only will draw something if a web camera is currently enabled in the live streaming settings
+	// and the user has activated it.  Also, the game may override this functionality entirely, and draw the web cam video itself.
+	IGameLiveStreaming::Get().DrawSimpleWebCamVideo( Canvas );
 }
 
 void UGameViewportClient::PeekTravelFailureMessages(UWorld* InWorld, ETravelFailure::Type FailureType, const FString& ErrorString)
diff --git a/Engine/Source/Runtime/Engine/Private/ScreenRendering.cpp b/Engine/Source/Runtime/Engine/Private/ScreenRendering.cpp
--- a/Engine/Source/Runtime/Engine/Private/ScreenRendering.cpp
+++ b/Engine/Source/Runtime/Engine/Private/ScreenRendering.cpp
@@ -15,88 +15,3 @@ TGlobalResource<FScreenVertexDeclaration> GScreenVertexDeclaration;
 IMPLEMENT_SHADER_TYPE(,FScreenPS,TEXT("ScreenPixelShader"),TEXT("Main"),SF_Pixel);
 IMPLEMENT_SHADER_TYPE(,FScreenVS,TEXT("ScreenVertexShader"),TEXT("Main"),SF_Vertex);
 IMPLEMENT_SHADER_TYPE(,FScreenVSForGS,TEXT("ScreenVertexShader"),TEXT("MainForGS"),SF_Vertex);
-
-FGlobalBoundShaderState ScreenBoundShaderState;
-
-/**
- * Draws a texture rectangle on the screen using normalized (-1 to 1) device coordinates.
- */
-void DrawScreenQuad(FRHICommandList& RHICmdList, float X0, float Y0, float U0, float V0, float X1, float Y1, float U1, float V1, const FTexture* Texture)
-{ 
-	// Set the screen shaders
-	TShaderMapRef<FScreenVS> ScreenVertexShader(GetGlobalShaderMap());
-	TShaderMapRef<FScreenPS> ScreenPixelShader(GetGlobalShaderMap());
-
-	RHICmdList.CheckIsNull(); // need new approach for "static FGlobalBoundShaderState" for parallel rendering
-	SetGlobalBoundShaderState(RHICmdList, ScreenBoundShaderState, GScreenVertexDeclaration.VertexDeclarationRHI, *ScreenVertexShader, *ScreenPixelShader);
-
-	ScreenPixelShader->SetParameters(RHICmdList, Texture);
-
-	// Generate the vertices used
-	FScreenVertex Vertices[4];
-
-	Vertices[0].Position.X = X1;
-	Vertices[0].Position.Y = Y0;
-	Vertices[0].UV.X       = U1;
-	Vertices[0].UV.Y       = V0;
-
-	Vertices[1].Position.X = X1;
-	Vertices[1].Position.Y = Y1;
-	Vertices[1].UV.X       = U1;
-	Vertices[1].UV.Y       = V1;
-
-	Vertices[2].Position.X = X0;
-	Vertices[2].Position.Y = Y0;
-	Vertices[2].UV.X       = U0;
-	Vertices[2].UV.Y       = V0;
-
-	Vertices[3].Position.X = X0;
-	Vertices[3].Position.Y = Y1;
-	Vertices[3].UV.X       = U0;
-	Vertices[3].UV.Y       = V1;
-
-	DrawPrimitiveUP(RHICmdList, PT_TriangleStrip, 2, Vertices, sizeof(Vertices[0]));
-}
-
-void DrawNormalizedScreenQuad(FRHICommandList& RHICmdList, float X, float Y, float U0, float V0, float SizeX, float SizeY, float U1, float V1, FIntPoint TargetSize, FTexture2DRHIRef TextureRHI)
-{
-	// Set the screen shaders 
-	TShaderMapRef<FScreenVS> ScreenVertexShader(GetGlobalShaderMap());
-	TShaderMapRef<FScreenPS> ScreenPixelShader(GetGlobalShaderMap());
-	RHICmdList.CheckIsNull(); // need new approach for "static FGlobalBoundShaderState" for parallel rendering
-	SetGlobalBoundShaderState(RHICmdList, ScreenBoundShaderState, GScreenVertexDeclaration.VertexDeclarationRHI, *ScreenVertexShader, *ScreenPixelShader);
-
-	/// ?
-	ScreenPixelShader->SetParameters(RHICmdList, TStaticSamplerState<SF_Bilinear>::GetRHI(), TextureRHI);
-
-	// Generate the vertices used
-	FScreenVertex Vertices[4];
-
-	Vertices[0].Position.X = X + SizeX;
-	Vertices[0].Position.Y = Y;
-	Vertices[0].UV.X       = U1;
-	Vertices[0].UV.Y       = V0;
-
-	Vertices[1].Position.X = X + SizeX;
-	Vertices[1].Position.Y = Y + SizeY;
-	Vertices[1].UV.X       = U1;
-	Vertices[1].UV.Y       = V1;
-
-	Vertices[2].Position.X = X;
-	Vertices[2].Position.Y = Y;
-	Vertices[2].UV.X       = U0;
-	Vertices[2].UV.Y       = V0;
-
-	Vertices[3].Position.X = X;
-	Vertices[3].Position.Y = Y + SizeY;
-	Vertices[3].UV.X       = U0;
-	Vertices[3].UV.Y       = V1;
-	
-	for(int32 VertexIndex = 0;VertexIndex < 4;VertexIndex++)
-	{
-		Vertices[VertexIndex].Position.X = -1.0f + 2.0f * (Vertices[VertexIndex].Position.X) / (float)TargetSize.X;
-		Vertices[VertexIndex].Position.Y = (+1.0f - 2.0f * (Vertices[VertexIndex].Position.Y) / (float)TargetSize.Y);
-	}
-
-	DrawPrimitiveUP(RHICmdList, PT_TriangleStrip, 2, Vertices, sizeof(Vertices[0]));
-}
diff --git a/Engine/Source/Runtime/Engine/Public/Features/ILiveStreamingService.h b/Engine/Source/Runtime/Engine/Public/Features/ILiveStreamingService.h
new file mode 100644
--- /dev/null
+++ b/Engine/Source/Runtime/Engine/Public/Features/ILiveStreamingService.h
@@ -0,0 +1,224 @@
+// Copyright 1998-2014 Epic Games, Inc. All Rights Reserved.
+
+#pragma once
+
+#include "Runtime/Core/Public/Features/IModularFeature.h"
+
+
+/**
+ * Contains information about the status of the live streaming service.  This is passed along to a user's
+ * function through the FLiveStreamingStatusCallback delegate
+ */
+struct FLiveStreamingStatus
+{
+	/** Types of statuses */
+	enum class EStatusType
+	{
+		/** An error of some sort */
+		Failure,
+
+		/** Progress message while we're getting setup to stream, such as login state */
+		Progress,
+
+		/** Success message that streaming is now live */
+		BroadcastStarted,
+
+		/** Sent when broadcasting stops */
+		BroadcastStopped,
+
+		/** Web cam is now active */
+		WebCamStarted,
+
+		/** Web cam is stopped */
+		WebCamStopped,
+
+		/** Web cam texture is ready (now safe to call GetWebCamTexture()) */
+		WebCamTextureReady,
+
+		/** Web cam texture is no longer valid (GetWebCamTexture() will return null) */
+		WebCamTextureNotReady,
+	};
+
+	/** Status type */
+	EStatusType StatusType;
+
+	/** Customized description of the current status from the live streaming service.  This string may include
+	    extra detail about the service's status, specific to that streaming system.  Unless the status is
+		an error condition, this text should be brief enough that it can be displayed as an editor notification.*/
+	FText CustomStatusDescription;
+};
+
+
+
+/**
+ * Defines parameters for initiating a new broadcast.  Keep in mind that different live streaming services have
+ * different restrictions on streaming buffer size and formats
+ */
+struct FBroadcastConfig
+{
+	/** Width of the streaming buffer.  If this changes, you'll need to stop and restart broadcasting.  Note that some streaming services
+	    require the Width to be a multiple of 32 (for SSE compression).  Call MakeValidVideoBufferResolution() first to make sure it's valid. */
+	int VideoBufferWidth;
+
+	/** Height of the streaming buffer.  If this changes, you'll need to stop and restart broadcasting.  Note that some streaming services
+	    require the height to be a multiple of 16 (for SSE compression).  Call MakeValidVideoBufferResolution() first to make sure it's valid. */
+	int VideoBufferHeight;
+
+	/** Target frame rate of this broadcast */
+	int FramesPerSecond;
+
+	/** Defines the possible pixel formats the output streaming buffer can be in */
+	enum class EBroadcastPixelFormat
+	{
+		/** Four bytes per pixel, red channel first, alpha last */
+		R8G8B8A8,
+
+		/** Four bytes per pixel, blue channel first, alpha last */
+		B8G8R8A8,
+
+		/** Four bytes per pixel, alpha first, blue channel last */
+		A8R8G8B8,
+
+		/** Four bytes per pixel, alpha first, red channel last */
+		A8B8G8R8,
+	};
+
+	/** Desired pixel format for the buffer. */
+	EBroadcastPixelFormat PixelFormat;
+
+	/** True to capture output audio from the local computer */
+	bool bCaptureAudioFromComputer;
+
+	/** True to capture audio from the default microphone */
+	bool bCaptureAudioFromMicrophone;
+
+
+	/** Default constructor */
+	FBroadcastConfig()
+		: VideoBufferWidth( 1280 ),
+		  VideoBufferHeight( 720 ),
+		  FramesPerSecond( 30 ),
+		  PixelFormat( EBroadcastPixelFormat::B8G8R8A8 ),
+		  bCaptureAudioFromComputer( true ),
+		  bCaptureAudioFromMicrophone( true )
+	{
+	}
+};
+
+
+
+/**
+ * Defines parameters for displaying a live web cam over your broadcast
+ */
+struct FWebCamConfig
+{
+	/** Desired web cam capture resolution width.  The web cam may only support a limited number of resolutions, so we'll
+		choose one that matches as closely to this as possible */
+	int DesiredWebCamWidth;
+
+	/** Desired web cam capture resolution height. */
+	int DesiredWebCamHeight;
+
+
+	/** Default constructor */
+	FWebCamConfig()
+		: DesiredWebCamWidth( 320 ),
+		  DesiredWebCamHeight( 240 )
+	{
+	}
+};
+
+
+/**
+ * Generic interface to an implementation of a live streaming system, usually capable of capturing video and audio
+ * and streaming live content to the internet
+ */
+class ILiveStreamingService : public IModularFeature
+{
+
+public:
+
+	DECLARE_MULTICAST_DELEGATE_OneParam( FOnStatusChanged, const FLiveStreamingStatus& );
+
+	/** Event to register with to find out when the status of the streaming system changes, such as when transitioning between
+		login phases, or a web cam becomes active */
+	virtual FOnStatusChanged& OnStatusChanged() = 0;
+
+
+	/**
+	 * Called to start broadcasting video and audio.  Note that the actual streaming may not begin immediately.
+	 *
+	 * @param	Config		Configuration for this broadcast.  If you want to change this later, you'll need to stop and restart the broadcast.
+	 */
+	virtual void StartBroadcasting( const FBroadcastConfig& Config ) = 0;
+
+	/** Called to stop broadcasting video and audio.  There may be a delay before streaming actually stops. */
+	virtual void StopBroadcasting() = 0;
+
+	/**
+	 * Checks to see if we've currently initiated a broadcast.  Note that even though we've asked to broadcast, the live stream
+	 * may not have actually started yet.  For example, we could still be waiting for the login authentication to complete
+	 *
+	 * @return	True if we're currently broadcasting
+	 */
+	virtual bool IsBroadcasting() const = 0;
+
+	/**
+	 * Checks to see if this streaming service is currently broadcasting and ready to accept new video frames.
+	 *
+	 * @return	True if we're broadcasting right now!
+	 */
+	virtual bool IsReadyForVideoFrames() const = 0;
+
+	/**
+	 * Given a desired streaming resolution, allows the streaming service to change the resolution to be valid for
+	 * broadcasting on this service.  Some streaming systems have specific requirements about resolutions, such as
+	 * multiples of 32 or 16.  Note that image aspect ratio can be affected by this call (that is, aspect ratio is
+	 * not always preserved, so it will be up to you to compensate.)
+	 *
+	 * @param	VideoBufferWidth	(In/Out) Set this to the width that you want, and you'll get back a width that it actually valid.
+	 * @param	VideoBufferHeight	(In/Out) Set this to the height that you want, and you'll get back a height that it actually valid.
+	 */
+	virtual void MakeValidVideoBufferResolution( int& VideoBufferWidth, int& VideoBufferHeight ) const = 0;
+
+	/**
+	 * Queries information about the current broadcast.  Should only be called while IsBroadcasting() is true.
+	 *
+	 * @param	BroadcastConfig		(Out) Filled with information about the current broadcast.
+	 */
+	virtual void QueryBroadcastConfig( FBroadcastConfig& OutBroadcastConfig ) const = 0;
+
+	/**
+	 * Broadcasts a new video frame.  Should only be called while IsReadyForVideoFrames() is true.  The frame buffer data must match
+	 * the format and resolution that the live streaming service expects.  You can call QueryBroadcastConfig() to find that out.
+	 *
+	 * @param	VideoFrameBuffer		The new video data to push to the live stream.  The size of the buffer must match the expected buffer resolution returned from QueryBroadcastConfig()!
+	 */
+	virtual void PushVideoFrame( const FColor* VideoFrameBuffer ) = 0;
+
+	/**
+	 * Enables broadcasting from your web camera device
+	 *
+	 * @param	Config	Settings for the web cam 
+	 */
+	virtual void StartWebCam( const FWebCamConfig& Config ) = 0;
+
+	/**
+	 * Turns off the web cam
+	 */
+	virtual void StopWebCam() = 0;
+
+	/**
+	 * Checks to see if the web cam is currently enabled
+	 *
+	 * @return	True if the web camera is active
+	 */
+	virtual bool IsWebCamEnabled() const = 0;
+
+	/**
+	 * For services that support web cams, this function will return the current picture on the web cam as a GPU texture
+	 *
+	 * @return	Current web cam frame as a 2D texture
+	 */
+	virtual class UTexture2D* GetWebCamTexture() = 0;
+};
\ No newline at end of file
diff --git a/Engine/Source/Runtime/Engine/Public/ScreenRendering.h b/Engine/Source/Runtime/Engine/Public/ScreenRendering.h
--- a/Engine/Source/Runtime/Engine/Public/ScreenRendering.h
+++ b/Engine/Source/Runtime/Engine/Public/ScreenRendering.h
@@ -127,9 +127,3 @@ public:
 	FScreenVSForGS() {}
 };
 
-/**
- * Draws a texture rectangle on the screen using normalized (-1 to 1) device coordinates.
- */
-extern void DrawScreenQuad(FRHICommandList& RHICmdList, float X0, float Y0, float U0, float V0, float X1, float Y1, float U1, float V1, const FTexture* Texture);
-
-extern void ENGINE_API DrawNormalizedScreenQuad(FRHICommandList& RHICmdList, float X0, float Y0, float U0, float V0, float X1, float Y1, float U1, float V1, FIntPoint TargetSize, FTexture2DRHIRef TextureRHI);
diff --git a/Engine/Source/Runtime/GameLiveStreaming/GameLiveStreaming.Build.cs b/Engine/Source/Runtime/GameLiveStreaming/GameLiveStreaming.Build.cs
new file mode 100644
--- /dev/null
+++ b/Engine/Source/Runtime/GameLiveStreaming/GameLiveStreaming.Build.cs
@@ -0,0 +1,30 @@
+// Copyright 1998-2014 Epic Games, Inc. All Rights Reserved.
+
+namespace UnrealBuildTool.Rules
+{
+	public class GameLiveStreaming : ModuleRules
+	{
+		public GameLiveStreaming( TargetInfo Target )
+		{
+			PublicDependencyModuleNames.AddRange(
+				new string[] 
+				{
+					"Core",
+				}
+			);
+
+			PrivateDependencyModuleNames.AddRange(
+				new string[] 
+				{
+					"CoreUObject",
+					"Engine",
+					"Slate",
+					"SlateCore",
+					"RenderCore",
+					"ShaderCore",
+					"RHI"
+				}
+			);
+		}
+	}
+}
diff --git a/Engine/Source/Runtime/GameLiveStreaming/GameLiveStreaming.cpp b/Engine/Source/Runtime/GameLiveStreaming/GameLiveStreaming.cpp
new file mode 100644
--- /dev/null
+++ b/Engine/Source/Runtime/GameLiveStreaming/GameLiveStreaming.cpp
@@ -0,0 +1,444 @@
+// Copyright 1998-2014 Epic Games, Inc. All Rights Reserved.
+#include "GameLiveStreamingModule.h"
+#include "GameLiveStreaming.h"
+#include "Runtime/Core/Public/Features/IModularFeatures.h"
+#include "Runtime/Engine/Public/Features/ILiveStreamingService.h"
+#include "Slate.h"
+#include "ScreenRendering.h"
+#include "RenderCore.h"
+#include "RHIStaticStates.h"
+
+IMPLEMENT_MODULE( FGameLiveStreaming, "GameLiveStreaming" );
+
+#define LOCTEXT_NAMESPACE "GameLiveStreaming"
+
+FGameLiveStreaming::FGameLiveStreaming()
+{
+	LiveStreamer = nullptr;
+	ReadbackTextureIndex = 0;
+	ReadbackBufferIndex = 0;
+	ReadbackBuffers[0] = nullptr;
+	ReadbackBuffers[1] = nullptr;
+	bDrawSimpleWebCamVideo = true;
+
+	// Setup some console commands
+	{
+		static FAutoConsoleCommand BroadcastStartCommand
+		(
+			TEXT( "Broadcast.Start" ),
+			TEXT( "Starts broadcasting game video and audio using a live internet streaming service\n" )
+			TEXT( "\n" )
+			TEXT( "Optional parameters:\n" )
+			TEXT( "     FrameRate=<int32>                   Frame rate to stream video when broadcasting\n")
+			TEXT( "     ScreenScaling=<float>               How much to scale the broadcast video resolution down\n" )
+			TEXT( "     bEnableWebCam=<bool>                Enables video from your web camera to be captured and displayed\n" )
+			TEXT( "     DesiredWebCamWidth=<int32>          Horizontal resolution for web camera capture\n" )
+			TEXT( "     DesiredWebCamHeight=<int32>         Veritcal resolution for web camera capture\n" )
+			TEXT( "     bCaptureAudioFromComputer=<bool>    Enables capturing sound that is played back through your PC\n" )
+			TEXT( "     bCaptureAudioFromMicrophone=<bool>  Enables capturing sound from your default microphone\n" )
+			TEXT( "     bDrawSimpleWebCamVideo=<bool>       Draws a simple web cam video on top of the viewport\n" ),
+			FConsoleCommandWithArgsDelegate::CreateStatic( 
+				[]( const TArray< FString >& Args )
+				{
+					FGameBroadcastConfig Config;
+					for( FString Arg : Args )
+					{
+						FParse::Value( *Arg, TEXT( "FrameRate="), Config.FrameRate );
+						FParse::Value( *Arg, TEXT( "ScreenScaling="), Config.ScreenScaling );
+						FParse::Bool( *Arg, TEXT( "bEnableWebCam="), Config.bEnableWebCam );
+						FParse::Value( *Arg, TEXT( "DesiredWebCamWidth="), Config.DesiredWebCamWidth );
+						FParse::Value( *Arg, TEXT( "DesiredWebCamHeight="), Config.DesiredWebCamHeight );
+						FParse::Bool( *Arg, TEXT( "bCaptureAudioFromComputer="), Config.bCaptureAudioFromComputer );
+						FParse::Bool( *Arg, TEXT( "bCaptureAudioFromMicrophone="), Config.bCaptureAudioFromMicrophone );
+						FParse::Bool( *Arg, TEXT( "bDrawSimpleWebCamVideo=" ), Config.bDrawSimpleWebCamVideo );
+					}
+					IGameLiveStreaming::Get().StartBroadcastingGame( Config );
+				} )
+		);
+
+		static FAutoConsoleCommand BroadcastStopCommand
+		(
+			TEXT( "Broadcast.Stop" ),
+			TEXT( "Stops broadcasting game video" ),
+			FConsoleCommandDelegate::CreateStatic( 
+				[]
+				{
+					IGameLiveStreaming::Get().StopBroadcastingGame();
+				} )
+		);
+	}
+}
+
+
+void FGameLiveStreaming::StartupModule()
+{
+}
+
+
+void FGameLiveStreaming::ShutdownModule()
+{
+	// Stop broadcasting at shutdown
+	StopBroadcastingGame();
+}
+
+
+
+bool FGameLiveStreaming::IsBroadcastingGame() const
+{
+	bool bAnythingBroadcasting = false;
+	if( LiveStreamer != nullptr )
+	{
+		if( LiveStreamer->IsBroadcasting() )
+		{
+			bAnythingBroadcasting = true;
+		}
+	}
+	return bAnythingBroadcasting;
+}
+
+
+void FGameLiveStreaming::StartBroadcastingGame( const FGameBroadcastConfig& GameBroadcastConfig )
+{
+	if( !IsBroadcastingGame() )
+	{
+		if( LiveStreamer != nullptr )
+		{
+			FSlateRenderer* SlateRenderer = FSlateApplication::Get().GetRenderer().Get();
+			SlateRenderer->OnSlateWindowRendered().RemoveAll( this );
+			LiveStreamer->OnStatusChanged().RemoveAll( this );
+			LiveStreamer = nullptr;
+		}
+
+		static const FName LiveStreamingFeatureName( "LiveStreaming" );
+		if( IModularFeatures::Get().IsModularFeatureAvailable( LiveStreamingFeatureName ) )
+		{
+			FSlateRenderer* SlateRenderer = FSlateApplication::Get().GetRenderer().Get();
+			SlateRenderer->OnSlateWindowRendered().AddRaw( this, &FGameLiveStreaming::OnSlateWindowRenderedDuringBroadcasting );
+
+			this->bDrawSimpleWebCamVideo = GameBroadcastConfig.bDrawSimpleWebCamVideo;
+
+			// Select a live streaming service
+			// @todo livestream: This will interfere with editor live streaming if both are running at the same time!  The editor live 
+			// streaming does check to make sure that game isn't already broadcasting, but the game currently doesn't have a good way to
+			// do that, besides asking the LiveStreamer itself.
+			// @todo livestream: Always using first live streaming plugin we find, but instead this should be a configurable option
+			LiveStreamer = static_cast<ILiveStreamingService*>( IModularFeatures::Get().GetModularFeatureImplementation( LiveStreamingFeatureName, 0 ) );
+
+			// Register to find out about status changes
+			LiveStreamer->OnStatusChanged().AddRaw( this, &FGameLiveStreaming::BroadcastStatusCallback );
+
+			UGameViewportClient* GameViewportClient = GEngine->GameViewport;
+			check( GameViewportClient != nullptr );
+
+			// @todo livestream: What about if viewport size changes while we are still broadcasting?  We need to restart the broadcast!
+			FBroadcastConfig BroadcastConfig;
+			BroadcastConfig.VideoBufferWidth = GameViewportClient->Viewport->GetSizeXY().X;
+			BroadcastConfig.VideoBufferHeight = GameViewportClient->Viewport->GetSizeXY().Y;
+
+			BroadcastConfig.VideoBufferWidth = FPlatformMath::FloorToInt( (float)BroadcastConfig.VideoBufferWidth * GameBroadcastConfig.ScreenScaling );
+			BroadcastConfig.VideoBufferHeight = FPlatformMath::FloorToInt( (float)BroadcastConfig.VideoBufferHeight * GameBroadcastConfig.ScreenScaling );
+
+			// Fix up the desired resolution so that it will work with the streaming system.  Some broadcasters require the
+			// video buffer to be multiples of specific values, such as 32
+			// @todo livestream: This could cause the aspect ratio to be changed and the buffer to be stretched non-uniformly, but usually the aspect only changes slightly
+			LiveStreamer->MakeValidVideoBufferResolution( BroadcastConfig.VideoBufferWidth, BroadcastConfig.VideoBufferHeight );
+
+			// Setup readback buffer textures
+			{
+				for( int32 TextureIndex = 0; TextureIndex < 2; ++TextureIndex )
+				{
+					FRHIResourceCreateInfo CreateInfo;
+					ReadbackTextures[ TextureIndex ] = RHICreateTexture2D(
+						BroadcastConfig.VideoBufferWidth,
+						BroadcastConfig.VideoBufferHeight,
+						PF_B8G8R8A8,
+						1,
+						1,
+						TexCreate_CPUReadback,
+						CreateInfo
+						);
+				}
+	
+				ReadbackTextureIndex = 0;
+
+				ReadbackBuffers[0] = nullptr;
+				ReadbackBuffers[1] = nullptr;
+				ReadbackBufferIndex = 0;
+			}
+
+			BroadcastConfig.FramesPerSecond = GameBroadcastConfig.FrameRate;
+			BroadcastConfig.PixelFormat = FBroadcastConfig::EBroadcastPixelFormat::B8G8R8A8;	// Matches viewport backbuffer format
+			BroadcastConfig.bCaptureAudioFromComputer = GameBroadcastConfig.bCaptureAudioFromComputer;
+			BroadcastConfig.bCaptureAudioFromMicrophone = GameBroadcastConfig.bCaptureAudioFromMicrophone;
+			LiveStreamer->StartBroadcasting( BroadcastConfig );
+
+			if( GameBroadcastConfig.bEnableWebCam )
+			{
+				// @todo livestream: Allow web cam to be started/stopped independently from the broadcast itself, so users can setup their web cam
+				FWebCamConfig WebCamConfig;
+				WebCamConfig.DesiredWebCamWidth = GameBroadcastConfig.DesiredWebCamWidth;
+				WebCamConfig.DesiredWebCamHeight = GameBroadcastConfig.DesiredWebCamHeight;
+				LiveStreamer->StartWebCam( WebCamConfig );
+			}
+		}
+	}
+}
+
+
+void FGameLiveStreaming::StopBroadcastingGame()
+{
+	if( LiveStreamer != nullptr )
+	{
+		LiveStreamer->OnStatusChanged().RemoveAll( this );
+
+		if( LiveStreamer->IsBroadcasting() )
+		{
+			LiveStreamer->StopWebCam();
+			LiveStreamer->StopBroadcasting();
+		}
+		if( FSlateApplication::IsInitialized() )	// During shutdown, Slate may have already been destroyed by the time our viewport gets cleaned up
+		{
+			FSlateRenderer* SlateRenderer = FSlateApplication::Get().GetRenderer().Get();
+			SlateRenderer->OnSlateWindowRendered().RemoveAll( this );
+		}
+		LiveStreamer = nullptr;
+
+		// Cleanup readback buffer textures
+		{
+			FlushRenderingCommands();
+
+			ReadbackTextures[0].SafeRelease();
+			ReadbackTextures[1].SafeRelease();
+			ReadbackTextureIndex = 0;
+			ReadbackBuffers[0] = nullptr;
+			ReadbackBuffers[1] = nullptr;
+			ReadbackBufferIndex = 0;
+		}
+	}
+}
+
+
+
+void FGameLiveStreaming::OnSlateWindowRenderedDuringBroadcasting( SWindow& SlateWindow, void* ViewportRHIPtr )
+{
+	// If we're streaming live video/audio, we'll go ahead and push new video frames here
+	UGameViewportClient* GameViewportClient = GEngine->GameViewport;
+	if( IsBroadcastingGame() && GameViewportClient != nullptr )
+	{
+		// We only care about our own Slate window
+		if( GameViewportClient->GetWindow() == SlateWindow.AsShared() )
+		{
+			// Check to see if we're streaming live video.  If so, we'll want to push new frames to be broadcast.
+			if( LiveStreamer->IsReadyForVideoFrames() )
+			{
+				// Check to see if there are any video frames ready to push
+				BroadcastGameVideoFrame();
+
+				// Start copying next game video frame
+				const FViewportRHIRef* ViewportRHI = ( const FViewportRHIRef* )ViewportRHIPtr;
+				StartCopyingNextGameVideoFrame( *ViewportRHI );
+			}
+		}
+	}
+	else
+	{
+		FSlateRenderer* SlateRenderer = FSlateApplication::Get().GetRenderer().Get();
+
+		// No longer broadcasting.  The live streaming service may have been interrupted.
+		LiveStreamer->OnStatusChanged().RemoveAll( this );
+		LiveStreamer = nullptr;
+		SlateRenderer->OnSlateWindowRendered().RemoveAll( this );
+	}
+}
+
+
+void FGameLiveStreaming::BroadcastGameVideoFrame()
+{
+	if( ReadbackBuffers[ ReadbackBufferIndex ] != nullptr )
+	{
+		// Great, we have a new frame back from the GPU.  Upload the video frame!
+		LiveStreamer->PushVideoFrame( (FColor*)ReadbackBuffers[ ReadbackBufferIndex ] );
+
+		// Unmap the buffer now that we've pushed out the frame
+		{
+			struct FReadbackFromStagingBufferContext
+			{
+				FGameLiveStreaming* This;
+			};
+			FReadbackFromStagingBufferContext ReadbackFromStagingBufferContext =
+			{
+				this
+			};
+			ENQUEUE_UNIQUE_RENDER_COMMAND_ONEPARAMETER(
+				ReadbackFromStagingBuffer,
+				FReadbackFromStagingBufferContext,Context,ReadbackFromStagingBufferContext,
+			{
+				RHIUnmapStagingSurface( Context.This->ReadbackTextures[ Context.This->ReadbackTextureIndex ] );
+			});
+		}								
+	}
+	else
+	{
+		// Frame hasn't finished copying from the GPU, so we can't broadcast a frame yet
+	}
+}
+
+
+void FGameLiveStreaming::StartCopyingNextGameVideoFrame( const FViewportRHIRef& ViewportRHI )
+{
+	// Check the video buffer size, in case we were configured to downscale the image
+	FBroadcastConfig BroadcastConfig;
+	LiveStreamer->QueryBroadcastConfig( BroadcastConfig );
+
+	const FIntPoint ResizeTo( BroadcastConfig.VideoBufferWidth, BroadcastConfig.VideoBufferHeight );
+
+	static const FName RendererModuleName( "Renderer" );
+	IRendererModule& RendererModule = FModuleManager::GetModuleChecked<IRendererModule>( RendererModuleName );
+
+	UGameViewportClient* GameViewportClient = GEngine->GameViewport;
+	check( GameViewportClient != nullptr );
+
+	struct FCopyVideoFrame
+	{
+		FViewportRHIRef ViewportRHI;
+		IRendererModule* RendererModule;
+		FIntPoint ResizeTo;
+		FGameLiveStreaming* This;
+	};
+	FCopyVideoFrame CopyVideoFrame =
+	{
+		ViewportRHI,
+		&RendererModule,
+		ResizeTo,
+		this
+	};
+
+	ENQUEUE_UNIQUE_RENDER_COMMAND_ONEPARAMETER(
+		ReadSurfaceCommand,
+		FCopyVideoFrame,Context,CopyVideoFrame,
+	{
+		FPooledRenderTargetDesc OutputDesc(FPooledRenderTargetDesc::Create2DDesc(Context.ResizeTo, PF_B8G8R8A8, TexCreate_None, TexCreate_RenderTargetable, false));
+			
+		TRefCountPtr<IPooledRenderTarget> ResampleTexturePooledRenderTarget;
+		Context.RendererModule->RenderTargetPoolFindFreeElement(OutputDesc, ResampleTexturePooledRenderTarget, TEXT("ResampleTexture"));
+		check( ResampleTexturePooledRenderTarget );
+
+		const FSceneRenderTargetItem& DestRenderTarget = ResampleTexturePooledRenderTarget->GetRenderTargetItem();
+
+		RHISetRenderTarget( DestRenderTarget.TargetableTexture, FTextureRHIRef() );
+		RHISetViewport(0, 0, 0.0f, Context.ResizeTo.X, Context.ResizeTo.Y, 1.0f);
+
+		RHISetBlendState(TStaticBlendState<>::GetRHI());
+		RHISetRasterizerState(TStaticRasterizerState<>::GetRHI());
+		RHISetDepthStencilState(TStaticDepthStencilState<false,CF_Always>::GetRHI());
+
+		FRHICommandList& RHICmdList = FRHICommandList::GetNullRef();
+
+		FTexture2DRHIRef ViewportBackBuffer = RHIGetViewportBackBuffer( Context.ViewportRHI );
+
+		TShaderMapRef<FScreenVS> VertexShader(GetGlobalShaderMap());
+		TShaderMapRef<FScreenPS> PixelShader(GetGlobalShaderMap());
+
+		static FGlobalBoundShaderState BoundShaderState;
+		SetGlobalBoundShaderState(RHICmdList, BoundShaderState, Context.RendererModule->GetFilterVertexDeclaration().VertexDeclarationRHI, *VertexShader, *PixelShader);
+
+		if( Context.ResizeTo != FIntPoint( ViewportBackBuffer->GetSizeX(), ViewportBackBuffer->GetSizeY() ) )
+		{
+			// We're scaling down the window, so use bilinear filtering
+			PixelShader->SetParameters(RHICmdList, TStaticSamplerState<SF_Bilinear>::GetRHI(), ViewportBackBuffer);
+		}
+		else
+		{
+			// Drawing 1:1, so no filtering needed
+			PixelShader->SetParameters(RHICmdList, TStaticSamplerState<SF_Point>::GetRHI(), ViewportBackBuffer);
+		}
+
+		Context.RendererModule->DrawRectangle(
+			RHICmdList,
+			0, 0,		// Dest X, Y
+			Context.ResizeTo.X, Context.ResizeTo.Y,	// Dest Width, Height
+			0, 0,		// Source U, V
+			1, 1,		// Source USize, VSize
+			Context.ResizeTo,		// Target buffer size
+			FIntPoint(1, 1),		// Source texture size
+			*VertexShader,
+			EDRF_Default);
+
+		// Asynchronously copy render target from GPU to CPU
+		const bool bKeepOriginalSurface = false;
+		RHICopyToResolveTarget(
+			DestRenderTarget.TargetableTexture,
+			Context.This->ReadbackTextures[ Context.This->ReadbackTextureIndex ],
+			bKeepOriginalSurface,
+			FResolveParams());
+	});
+							
+
+	// Start mapping the newly-rendered buffer
+	{
+		struct FReadbackFromStagingBufferContext
+		{
+			FGameLiveStreaming* This;
+		};
+		FReadbackFromStagingBufferContext ReadbackFromStagingBufferContext =
+		{
+			this
+		};
+		ENQUEUE_UNIQUE_RENDER_COMMAND_ONEPARAMETER(
+			ReadbackFromStagingBuffer,
+			FReadbackFromStagingBufferContext,Context,ReadbackFromStagingBufferContext,
+		{
+			int32 UnusedWidth = 0;
+			int32 UnusedHeight = 0;
+			RHIMapStagingSurface( Context.This->ReadbackTextures[ Context.This->ReadbackTextureIndex ], Context.This->ReadbackBuffers[ Context.This->ReadbackBufferIndex ], UnusedWidth, UnusedHeight );
+
+			// Ping pong between readback textures
+			Context.This->ReadbackTextureIndex = ( Context.This->ReadbackTextureIndex + 1 ) % 2;
+		});
+	}
+
+	// Ping pong between readback buffers
+	ReadbackBufferIndex = ( ReadbackBufferIndex + 1 ) % 2;
+}
+
+
+void FGameLiveStreaming::DrawSimpleWebCamVideo( UCanvas* Canvas )
+{
+	if( IsBroadcastingGame() && bDrawSimpleWebCamVideo && LiveStreamer->IsWebCamEnabled() )
+	{
+		UTexture2D* WebCamTexture = LiveStreamer->GetWebCamTexture();
+		if( WebCamTexture != nullptr )
+		{
+			const float BorderPadding = 6.0f;
+			Canvas->Canvas->DrawTile( 
+				Canvas->SizeX - WebCamTexture->GetSizeX() - BorderPadding, BorderPadding,	// Top right justify
+				WebCamTexture->GetSizeX(), WebCamTexture->GetSizeY(),
+				0, 0, 
+				1.0f, 1.0f, 
+				FLinearColor::White, 
+				WebCamTexture->Resource, 
+				false );	// Alpha blend?
+		}
+	}
+}
+
+
+UTexture2D* FGameLiveStreaming::GetWebCamTexture()
+{
+	UTexture2D* WebCamTexture = nullptr;
+	if( IsBroadcastingGame() && bDrawSimpleWebCamVideo && LiveStreamer->IsWebCamEnabled() )
+	{
+		WebCamTexture = LiveStreamer->GetWebCamTexture();
+	}
+	return WebCamTexture;
+}
+
+	
+void FGameLiveStreaming::BroadcastStatusCallback( const FLiveStreamingStatus& Status )
+{
+	// @todo livestream: Hook this up to provide C++ users and Blueprints with status about the live streaming startup/webcam/shutdown
+}
+
+
+
+#undef LOCTEXT_NAMESPACE
diff --git a/Engine/Source/Runtime/GameLiveStreaming/GameLiveStreaming.h b/Engine/Source/Runtime/GameLiveStreaming/GameLiveStreaming.h
new file mode 100644
--- /dev/null
+++ b/Engine/Source/Runtime/GameLiveStreaming/GameLiveStreaming.h
@@ -0,0 +1,71 @@
+// Copyright 1998-2014 Epic Games, Inc. All Rights Reserved.
+#pragma once
+
+#include "IGameLiveStreaming.h"
+
+/**
+ * Implements support for broadcasting gameplay to a live internet stream
+ */
+class FGameLiveStreaming
+	: public IGameLiveStreaming
+{
+
+public:
+
+	/** Default constructor */
+	FGameLiveStreaming();
+
+	/** IModuleInterface overrides */
+	virtual void StartupModule() override;
+	virtual void ShutdownModule() override;
+
+	/** IGameLiveStreaming overrides */
+	virtual bool IsBroadcastingGame() const override;
+	virtual void StartBroadcastingGame( const FGameBroadcastConfig& GameBroadcastConfig ) override;
+	virtual void StopBroadcastingGame() override;
+	virtual void DrawSimpleWebCamVideo( UCanvas* Canvas ) override;
+	virtual class UTexture2D* GetWebCamTexture() override;
+
+protected:
+
+	/** Called by the live streaming service when streaming status has changed or an error has occurred */
+	void BroadcastStatusCallback( const struct FLiveStreamingStatus& Status );
+
+	/**
+	 * Called by the Slate rendered on the game thread right after a window has been rendered, to allow us to
+	 * capture frame buffer content after everything has been rendered, including debug canvas graphics
+	 *
+	 * @param	SlateWindow				The window that was just rendered
+	 * @param	ViewportRHIPtr			The viewport we rendered this window to (cast to FViewportRHIRef*)
+	 */
+	void OnSlateWindowRenderedDuringBroadcasting( SWindow& SlateWindow, void* ViewportRHIPtr );
+
+	/** Broadcasts a new video frame, only if the frame's data has finished copying from the GPU to system memory */
+	void BroadcastGameVideoFrame();
+
+	/** Starts copying the last rendered game frame from the GPU back buffer to a mapped system memory texture */
+	void StartCopyingNextGameVideoFrame( const FViewportRHIRef& ViewportRHI );
+
+
+
+private:
+
+	/** The live streaming service we're using.  Only valid while broadcasting. */
+	class ILiveStreamingService* LiveStreamer;
+
+	/** Static: Readback textures for asynchronously reading the viewport frame buffer back to the CPU.  We ping-pong between the buffers to avoid stalls. */
+	FTexture2DRHIRef ReadbackTextures[2];
+	
+	/** Static: We ping pong between the textures in case the GPU is a frame behind (GSystemSettings.bAllowOneFrameThreadLag) */
+	int32 ReadbackTextureIndex;
+
+	/** Static: Pointers to mapped system memory readback textures that game frames will be asynchronously copied to */
+	void* ReadbackBuffers[2];
+
+	/** The current buffer index.  We bounce between them to avoid stalls. */
+	int32 ReadbackBufferIndex;
+
+	/** True if we should draw a simple web cam video on top of the viewport while broadcasting */
+	bool bDrawSimpleWebCamVideo;
+
+};
diff --git a/Engine/Source/Runtime/GameLiveStreaming/GameLiveStreamingFunctionLibrary.cpp b/Engine/Source/Runtime/GameLiveStreaming/GameLiveStreamingFunctionLibrary.cpp
new file mode 100644
--- /dev/null
+++ b/Engine/Source/Runtime/GameLiveStreaming/GameLiveStreamingFunctionLibrary.cpp
@@ -0,0 +1,51 @@
+// Copyright 1998-2014 Epic Games, Inc. All Rights Reserved.
+#include "GameLiveStreamingModule.h"
+#include "GameLiveStreamingFunctionLibrary.h"
+#include "Public/IGameLiveStreaming.h"
+
+#define LOCTEXT_NAMESPACE "GameLiveStreaming"
+
+
+UGameLiveStreamingFunctionLibrary::UGameLiveStreamingFunctionLibrary( const FPostConstructInitializeProperties& PCIP )
+	: Super(PCIP)
+{
+}
+
+
+bool UGameLiveStreamingFunctionLibrary::IsBroadcastingGame()
+{
+	return IGameLiveStreaming::Get().IsBroadcastingGame();
+}
+
+
+void UGameLiveStreamingFunctionLibrary::StartBroadcastingGame(
+		int32 FrameRate,
+		float ScreenScaling,
+		bool bEnableWebCam,
+		int32 DesiredWebCamWidth,
+		int32 DesiredWebCamHeight,
+		bool bCaptureAudioFromComputer,
+		bool bCaptureAudioFromMicrophone,
+		bool bDrawSimpleWebCamVideo )
+{
+	FGameBroadcastConfig Config;
+	Config.FrameRate = FrameRate;
+	Config.ScreenScaling = ScreenScaling;
+	Config.bEnableWebCam = bEnableWebCam;
+	Config.DesiredWebCamWidth = DesiredWebCamWidth;
+	Config.DesiredWebCamHeight = DesiredWebCamHeight;
+	Config.bCaptureAudioFromComputer = bCaptureAudioFromComputer;
+	Config.bCaptureAudioFromMicrophone = bCaptureAudioFromMicrophone;
+	Config.bDrawSimpleWebCamVideo = bDrawSimpleWebCamVideo;
+	IGameLiveStreaming::Get().StartBroadcastingGame( Config );
+}
+
+
+void UGameLiveStreamingFunctionLibrary::StopBroadcastingGame()
+{
+	IGameLiveStreaming::Get().StopBroadcastingGame();
+}
+
+
+
+#undef LOCTEXT_NAMESPACE
diff --git a/Engine/Source/Runtime/GameLiveStreaming/GameLiveStreamingModule.h b/Engine/Source/Runtime/GameLiveStreaming/GameLiveStreamingModule.h
new file mode 100644
--- /dev/null
+++ b/Engine/Source/Runtime/GameLiveStreaming/GameLiveStreamingModule.h
@@ -0,0 +1,5 @@
+// Copyright 1998-2014 Epic Games, Inc. All Rights Reserved.
+#pragma once
+
+#include "Engine.h"
+
diff --git a/Engine/Source/Runtime/GameLiveStreaming/Public/GameLiveStreamingFunctionLibrary.h b/Engine/Source/Runtime/GameLiveStreaming/Public/GameLiveStreamingFunctionLibrary.h
new file mode 100644
--- /dev/null
+++ b/Engine/Source/Runtime/GameLiveStreaming/Public/GameLiveStreamingFunctionLibrary.h
@@ -0,0 +1,46 @@
+// Copyright 1998-2014 Epic Games, Inc. All Rights Reserved.
+#pragma once
+#include "GameLiveStreamingFunctionLibrary.generated.h"
+
+UCLASS()
+class UGameLiveStreamingFunctionLibrary : public UBlueprintFunctionLibrary
+{
+	GENERATED_UCLASS_BODY()
+
+
+	/**
+	 * Checks to see if we are currently broadcasting live video (and possibly audio) from the game's viewport
+	 *
+	 * @return	True if we are currently transmitting
+	 */
+	UFUNCTION( BlueprintPure, Category="Broadcasting" )
+	static bool IsBroadcastingGame();
+
+	/**
+	 * Starts broadcasting the game's video (and optionally audio) using an internet streaming service, if one is available 
+	 *
+	 * @param	FrameRate					Frame rate to stream video from when broadcasting to services like Twitch.
+	 * @param	ScreenScaling				How much to scale the broadcast video resolution down to reduce streaming bandwidth.  We recommend broadcasting at resolutions of 1280x720 or lower.  Some live streaming providers will not be able to transcode your video to a lower resolution, so using a high resolution stream may prevent low-bandwidth users from having a good viewing experience.
+	 * @param	bEnableWebCam				If enabled, video from your web camera will be captured and displayed while broadcasting, so that your viewers can see your presence.
+	 * @param	DesiredWebCamWidth			Desired web cam capture resolution width.  The web cam may only support a limited number of resolutions, so we'll choose one that matches as closely to this as possible
+	 * @param	DesiredWebCamHeight			Desired web cam capture resolution height.
+	 * @param	bCaptureAudioFromComputer	Enables broadcast of audio being played by your computer, such as in-game sounds
+	 * @param	bCaptureAudioFromMicrophone	Enables broadcast of audio from your default microphone recording device
+	 * @param	bDrawSimpleWebCamVideo		If enabled, the engine will draw a simple web cam image on top of the game viewport.  If you turn this off, it's up to you to draw the web cam image yourself.  You can access the web cam texture by calling IGameLiveStreaming::Get().GetWebCamTexture(). 
+	 */
+	UFUNCTION( BlueprintCallable, Category="Broadcasting" )
+	static void StartBroadcastingGame(
+		int32 FrameRate = 30,
+		float ScreenScaling = 1.0f,
+		bool bEnableWebCam = true,
+		int32 DesiredWebCamWidth = 320,
+		int32 DesiredWebCamHeight = 240,
+		bool bCaptureAudioFromComputer = true,
+		bool bCaptureAudioFromMicrophone = true,
+		bool bDrawSimpleWebCamVideo = true );
+
+	/** Stops broadcasting the game */
+	UFUNCTION( BlueprintCallable, Category="Broadcasting" )
+	static void StopBroadcastingGame();
+
+};
diff --git a/Engine/Source/Runtime/GameLiveStreaming/Public/IGameLiveStreaming.h b/Engine/Source/Runtime/GameLiveStreaming/Public/IGameLiveStreaming.h
new file mode 100644
--- /dev/null
+++ b/Engine/Source/Runtime/GameLiveStreaming/Public/IGameLiveStreaming.h
@@ -0,0 +1,111 @@
+// Copyright 1998-2014 Epic Games, Inc. All Rights Reserved.
+#pragma once
+
+#include "ModuleManager.h"
+
+
+/**
+ * Settings for a live game broadcast
+ */
+struct FGameBroadcastConfig
+{
+	/** Frame rate to stream video when broadcasting to services like Twitch. */
+	int32 FrameRate;
+
+	/** How much to scale the broadcast video resolution down to reduce streaming bandwidth.  We recommend broadcasting at resolutions of 1280x720 or lower.  Some live streaming providers will not be able to transcode your video to a lower resolution, so using a high resolution stream may prevent low-bandwidth users from having a good viewing experience.*/
+	float ScreenScaling;
+
+	/** If enabled, video from your web camera will be captured and displayed while broadcasting, so that your viewers can see your presence. */
+	bool bEnableWebCam;
+
+	/** Desired web cam capture resolution width.  The web cam may only support a limited number of resolutions, so we'll
+		choose one that matches as closely to this as possible */
+	int32 DesiredWebCamWidth;
+
+	/** Desired web cam capture resolution height. */
+	int32 DesiredWebCamHeight;
+
+	/** Enables broadcast of audio being played by your computer, such as in-game sounds */
+	bool bCaptureAudioFromComputer;
+
+	/** Enables broadcast of audio from your default microphone recording device */
+	bool bCaptureAudioFromMicrophone;
+
+	/** If enabled, the engine will draw a simple web cam image on top of the game viewport.  If you turn this off, it's up to you to draw the web cam image yourself.  You can access the web cam texture by calling IGameLiveStreaming::Get().GetWebCamTexture(). */
+	bool bDrawSimpleWebCamVideo;
+
+
+	/**
+	 * Default constructor
+	 */
+	FGameBroadcastConfig()
+		: FrameRate( 30 ),
+		  ScreenScaling( 1.0f ),
+		  bEnableWebCam( true ),
+		  DesiredWebCamWidth( 320 ),
+		  DesiredWebCamHeight( 240 ),
+		  bCaptureAudioFromComputer( true ),
+		  bCaptureAudioFromMicrophone( true ),
+		  bDrawSimpleWebCamVideo( true )
+	{
+	}
+};
+
+
+/**
+ * Interface to the GameLiveStreaming module's functionality for broadcasting live gameplay to the internet
+ */
+class IGameLiveStreaming
+	: public IModuleInterface
+{
+
+public:
+
+	/**
+	 * Checks to see if we are currently broadcasting live video (and possibly audio) from the game's viewport
+	 *
+	 * @return	True if we are currently transmitting
+	 */
+	virtual bool IsBroadcastingGame() const = 0;
+
+	/**
+	 * Starts broadcasting the game's video (and optionally audio) using an internet streaming service, if one is available 
+	 *
+	 * @param	Config	Settings for this game broadcast
+	 */
+	virtual void StartBroadcastingGame( const FGameBroadcastConfig& GameBroadcastConfig ) = 0;
+
+	/** Stops broadcasting the game */
+	virtual void StopBroadcastingGame() = 0;
+
+	/** 
+	 * Call this in your game's UI rendering code to draw a simple default web cam image while broadcasting.  Instead of
+	 * using this, you can grab the web cam texture yourself by calling GetWebCamTexture() and draw it however you want!
+	 *
+	 * @param	Canvas	The canvas to draw the web cam image to
+	 */
+	virtual void DrawSimpleWebCamVideo( UCanvas* Canvas ) = 0;
+
+	/**
+	 * Exposes access to the web cam video feed as a 2D texture
+	 *
+	 * @return	The web cam video feed texture, or null if the web cam isn't enabled
+	 */
+	virtual class UTexture2D* GetWebCamTexture() = 0;
+
+
+public:
+
+	/**
+	 * Singleton-like access to IGameLiveStreaming
+	 *
+	 * @return Returns instance of IGameLiveStreaming object
+	 */
+	static inline IGameLiveStreaming& Get()
+	{
+		return FModuleManager::LoadModuleChecked< IGameLiveStreaming >( "GameLiveStreaming" );
+	}
+
+};
+
+
diff --git a/Engine/Source/Runtime/Renderer/Private/PostProcess/SceneFilterRendering.h b/Engine/Source/Runtime/Renderer/Private/PostProcess/SceneFilterRendering.h
--- a/Engine/Source/Runtime/Renderer/Private/PostProcess/SceneFilterRendering.h
+++ b/Engine/Source/Runtime/Renderer/Private/PostProcess/SceneFilterRendering.h
@@ -11,46 +11,6 @@
 
 #include "SceneRenderTargets.h"
 
-/** The vertex data used to filter a texture. */
-struct FFilterVertex
-{
-	FVector4 Position;
-	FVector2D UV;
-};
-
-/** The filter vertex declaration resource type. */
-class FFilterVertexDeclaration : public FRenderResource
-{
-public:
-	FVertexDeclarationRHIRef VertexDeclarationRHI;
-
-	/** Destructor. */
-	virtual ~FFilterVertexDeclaration() {}
-
-	virtual void InitRHI()
-	{
-		FVertexDeclarationElementList Elements;
-		uint32 Stride = sizeof(FFilterVertex);
-		Elements.Add(FVertexElement(0,STRUCT_OFFSET(FFilterVertex,Position),VET_Float4,0,Stride));
-		Elements.Add(FVertexElement(0,STRUCT_OFFSET(FFilterVertex,UV),VET_Float2,1,Stride));
-		VertexDeclarationRHI = RHICreateVertexDeclaration(Elements);
-	}
-
-	virtual void ReleaseRHI()
-	{
-		VertexDeclarationRHI.SafeRelease();
-	}
-};
-
-// use r.DrawDenormalizedQuadMode to override the function call setting (quick way to see if an artifact is caused why this optimization)
-enum EDrawRectangleFlags
-{
-	// Rectangle is created by 2 triangles (diagonal can cause some slightly less efficient shader execution), this is the default as it has no artifacts
-	EDRF_Default,
-	//
-	EDRF_UseTriangleOptimization
-};
-
 /**
  * Draws a quad with the given vertex positions and UVs in denormalized pixel/texel coordinates.
  * The platform-dependent mapping from pixels to texels is done automatically.
@@ -78,7 +38,7 @@ extern void DrawRectangle(
 	float SizeV,
 	FIntPoint TargetSize,
 	FIntPoint TextureSize,
-	FShader* VertexShader,
+	class FShader* VertexShader,
 	EDrawRectangleFlags Flags = EDRF_Default
 	);
 
diff --git a/Engine/Source/Runtime/Renderer/Private/RendererPrivate.h b/Engine/Source/Runtime/Renderer/Private/RendererPrivate.h
--- a/Engine/Source/Runtime/Renderer/Private/RendererPrivate.h
+++ b/Engine/Source/Runtime/Renderer/Private/RendererPrivate.h
@@ -12,6 +12,7 @@
 #include "RHI.h"
 #include "RHIStaticStates.h"
 #include "ScenePrivate.h"
+#include "RendererInterface.h"
 
 DECLARE_LOG_CATEGORY_EXTERN(LogRenderer, Log, All);
 
@@ -37,6 +38,22 @@ public:
 	virtual void QueryVisualizeTexture(FQueryVisualizeTexureInfo& Out) override;
 	virtual void ExecVisualizeTextureCmd(const FString& Cmd) override;
 	virtual void UpdateMapNeedsLightingFullyRebuiltState(UWorld* World) override;
+	virtual void DrawRectangle(
+		FRHICommandList& RHICmdList,
+		float X,
+		float Y,
+		float SizeX,
+		float SizeY,
+		float U,
+		float V,
+		float SizeU,
+		float SizeV,
+		FIntPoint TargetSize,
+		FIntPoint TextureSize,
+		class FShader* VertexShader,
+		EDrawRectangleFlags Flags = EDRF_Default
+		) override;
+	virtual TGlobalResource<FFilterVertexDeclaration>& GetFilterVertexDeclaration() override;
 
 	virtual const TSet<FSceneInterface*>& GetAllocatedScenes() override
 	{
diff --git a/Engine/Source/Runtime/Renderer/Private/SceneRendering.cpp b/Engine/Source/Runtime/Renderer/Private/SceneRendering.cpp
--- a/Engine/Source/Runtime/Renderer/Private/SceneRendering.cpp
+++ b/Engine/Source/Runtime/Renderer/Private/SceneRendering.cpp
@@ -1023,6 +1023,32 @@ void FRendererModule::UpdateMapNeedsLightingFullyRebuiltState(UWorld* World)
 	World->SetMapNeedsLightingFullyRebuilt(World->Scene->GetRenderScene()->NumUncachedStaticLightingInteractions);
 }
 
+void FRendererModule::DrawRectangle(
+		FRHICommandList& RHICmdList,
+		float X,
+		float Y,
+		float SizeX,
+		float SizeY,
+		float U,
+		float V,
+		float SizeU,
+		float SizeV,
+		FIntPoint TargetSize,
+		FIntPoint TextureSize,
+		class FShader* VertexShader,
+		EDrawRectangleFlags Flags
+		)
+{
+	::DrawRectangle( RHICmdList, X, Y, SizeX, SizeY, U, V, SizeU, SizeV, TargetSize, TextureSize, VertexShader, Flags );
+}
+
+
+TGlobalResource<FFilterVertexDeclaration>& FRendererModule::GetFilterVertexDeclaration()
+{
+	return GFilterVertexDeclaration;
+}
+
+
 bool IsMobileHDR()
 {
 	static auto* MobileHDRCvar = IConsoleManager::Get().FindTConsoleVariableDataInt(TEXT("r.MobileHDR"));
diff --git a/Engine/Source/Runtime/Renderer/Public/RendererInterface.h b/Engine/Source/Runtime/Renderer/Public/RendererInterface.h
--- a/Engine/Source/Runtime/Renderer/Public/RendererInterface.h
+++ b/Engine/Source/Runtime/Renderer/Public/RendererInterface.h
@@ -371,6 +371,50 @@ struct FQueryVisualizeTexureInfo
 	TArray<FString> Entries;
 };
 
+
+/** The vertex data used to filter a texture. */
+struct FFilterVertex
+{
+	FVector4 Position;
+	FVector2D UV;
+};
+
+/** The filter vertex declaration resource type. */
+class FFilterVertexDeclaration : public FRenderResource
+{
+public:
+	FVertexDeclarationRHIRef VertexDeclarationRHI;
+
+	/** Destructor. */
+	virtual ~FFilterVertexDeclaration() {}
+
+	virtual void InitRHI()
+	{
+		FVertexDeclarationElementList Elements;
+		uint32 Stride = sizeof(FFilterVertex);
+		Elements.Add(FVertexElement(0,STRUCT_OFFSET(FFilterVertex,Position),VET_Float4,0,Stride));
+		Elements.Add(FVertexElement(0,STRUCT_OFFSET(FFilterVertex,UV),VET_Float2,1,Stride));
+		VertexDeclarationRHI = RHICreateVertexDeclaration(Elements);
+	}
+
+	virtual void ReleaseRHI()
+	{
+		VertexDeclarationRHI.SafeRelease();
+	}
+};
+
+// use r.DrawDenormalizedQuadMode to override the function call setting (quick way to see if an artifact is caused why this optimization)
+enum EDrawRectangleFlags
+{
+	// Rectangle is created by 2 triangles (diagonal can cause some slightly less efficient shader execution), this is the default as it has no artifacts
+	EDRF_Default,
+	//
+	EDRF_UseTriangleOptimization
+};
+
+
+
+
 /**
  * The public interface of the renderer module.
  */
@@ -430,6 +474,40 @@ public:
 	virtual void ExecVisualizeTextureCmd(const FString& Cmd) = 0;
 
 	virtual void UpdateMapNeedsLightingFullyRebuiltState(UWorld* World) = 0;
+
+	/**
+	 * Draws a quad with the given vertex positions and UVs in denormalized pixel/texel coordinates.
+	 * The platform-dependent mapping from pixels to texels is done automatically.
+	 * Note that the positions are affected by the current viewport.
+	 * NOTE: DrawRectangle should be used in the vertex shader to calculate the correct position and uv for vertices.
+	 *
+	 * X, Y							Position in screen pixels of the top left corner of the quad
+	 * SizeX, SizeY					Size in screen pixels of the quad
+	 * U, V							Position in texels of the top left corner of the quad's UV's
+	 * SizeU, SizeV					Size in texels of the quad's UV's
+	 * TargetSizeX, TargetSizeY		Size in screen pixels of the target surface
+	 * TextureSize                  Size in texels of the source texture
+	 * VertexShader					The vertex shader used for rendering
+	 * Flags						see EDrawRectangleFlags
+	 */
+	virtual void DrawRectangle(
+		FRHICommandList& RHICmdList,
+		float X,
+		float Y,
+		float SizeX,
+		float SizeY,
+		float U,
+		float V,
+		float SizeU,
+		float SizeV,
+		FIntPoint TargetSize,
+		FIntPoint TextureSize,
+		class FShader* VertexShader,
+		EDrawRectangleFlags Flags = EDRF_Default
+		) = 0;
+
+	/** @return Returns a vertex declaration that can be used with with the DrawRectangle() function */
+	virtual TGlobalResource<FFilterVertexDeclaration>& GetFilterVertexDeclaration() = 0;
 };
 
 
diff --git a/Engine/Source/Runtime/SlateCore/Public/Rendering/SlateRenderer.h b/Engine/Source/Runtime/SlateCore/Public/Rendering/SlateRenderer.h
--- a/Engine/Source/Runtime/SlateCore/Public/Rendering/SlateRenderer.h
+++ b/Engine/Source/Runtime/SlateCore/Public/Rendering/SlateRenderer.h
@@ -69,16 +69,32 @@ public:
 	 */
 	virtual void DrawWindows() {}
 	
+	/**
+	 * You must call this before calling CopyWindowsToVirtualScreenBuffer(), to setup the render targets first.
+	 * 
+	 * @param	bPrimaryWorkAreaOnly	True if we should capture only the primary monitor's work area, or false to capture the entire desktop spanning all monitors
+	 * @param	ScreenScaling	How much to downscale the desktop size
+	 * @param	LiveStreamingService	Optional pointer to a live streaming service this buffer needs to work with
+	 *
+	 * @return	The virtual screen rectangle.  The size of this rectangle will be the size of the render target buffer.
+	 */
+	virtual FIntRect SetupVirtualScreenBuffer( const bool bPrimaryWorkAreaOnly, const float ScreenScaling, class ILiveStreamingService* LiveStreamingService) { return FIntRect(); }
+
+
 	/**
 	 * Copies all slate windows out to a buffer at half resolution with debug information
 	 * like the mouse cursor and any keypresses.
 	 */
-	virtual void CopyWindowsToDrawBuffer(const TArray<FString>& KeypressBuffer) {}
+	virtual void CopyWindowsToVirtualScreenBuffer(const TArray<FString>& KeypressBuffer) {}
 	
 	/** Allows and disallows access to the crash tracker buffer data on the CPU */
-	virtual void MapCrashTrackerBuffer(void** OutImageData, int32* OutWidth, int32* OutHeight) {}
-	virtual void UnmapCrashTrackerBuffer() {}
+	virtual void MapVirtualScreenBuffer(void** OutImageData) {}
+	virtual void UnmapVirtualScreenBuffer() {}
 	
+	/** Callback that fires after Slate has rendered each window, each frame */
+	DECLARE_MULTICAST_DELEGATE_TwoParams( FOnSlateWindowRendered, SWindow&, void* );
+	FOnSlateWindowRendered& OnSlateWindowRendered() { return SlateWindowRendered; }
+
 	/** 
 	 * Sets which color vision filter to use
 	 */
@@ -182,6 +198,9 @@ protected:
 
 	TSharedPtr<class FSlateFontCache> FontCache;
 	TSharedPtr<class FSlateFontMeasure> FontMeasure;
+
+	/** Callback that fires after Slate has rendered each window, each frame */
+	FOnSlateWindowRendered SlateWindowRendered;
 };
 
 
diff --git a/Engine/Source/Runtime/SlateCore/Public/Widgets/SWindow.h b/Engine/Source/Runtime/SlateCore/Public/Widgets/SWindow.h
--- a/Engine/Source/Runtime/SlateCore/Public/Widgets/SWindow.h
+++ b/Engine/Source/Runtime/SlateCore/Public/Widgets/SWindow.h
@@ -605,24 +605,6 @@ public:
 	/** Sets the delegate to call when switching worlds in before ticking,drawing, or sending messages to widgets in this window */
 	void SetOnWorldSwitchHack( FOnSwitchWorldHack& InOnWorldSwitchHack );
 
-	/**
-	 * Static: Marks the specified window as drawn
-	 *
-	 * @param	Window	The window to mark as drawn
-	 */
-	static void MarkWindowAsDrawn( SWindow* Window )
-	{
-		Window->MarkAsDrawn();
-	}
-
-	/**
-	 * @return Returns a delegate that can be used to mark window as drawn. The delegate is safe to call from a different thread.
-	 */
-	FSimpleDelegate MakeMarkWindowAsDrawnDelegate()
-	{
-		return FSimpleDelegate::CreateStatic( &MarkWindowAsDrawn, this );
-	}
-
 	/**
 	 * Hack to switch worlds
 	 *
diff --git a/Engine/Source/Runtime/SlateRHIRenderer/Private/SlateRHIRenderer.cpp b/Engine/Source/Runtime/SlateRHIRenderer/Private/SlateRHIRenderer.cpp
--- a/Engine/Source/Runtime/SlateRHIRenderer/Private/SlateRHIRenderer.cpp
+++ b/Engine/Source/Runtime/SlateRHIRenderer/Private/SlateRHIRenderer.cpp
@@ -1,19 +1,13 @@
 // Copyright 1998-2014 Epic Games, Inc. All Rights Reserved.
 
 #include "SlateRHIRendererPrivatePCH.h"
-
+#include "Runtime/Engine/Public/Features/ILiveStreamingService.h"
 
 DECLARE_CYCLE_STAT(TEXT("Map Staging Buffer"),STAT_MapStagingBuffer,STATGROUP_CrashTracker);
 DECLARE_CYCLE_STAT(TEXT("Generate Capture Buffer"),STAT_GenerateCaptureBuffer,STATGROUP_CrashTracker);
 DECLARE_CYCLE_STAT(TEXT("Unmap Staging Buffer"),STAT_UnmapStagingBuffer,STATGROUP_CrashTracker);
 
 
-namespace CrashTrackerConstants
-{
-	static const float ScreenScaling = 0.5f;
-}
-
-
 // Defines the maximum size that a slate viewport will create
 #define MAX_VIEWPORT_SIZE 16384
 
@@ -44,8 +38,8 @@ void FSlateCrashReportResource::InitDynamicRHI()
 {
 	FRHIResourceCreateInfo CreateInfo;
 	CrashReportBuffer = RHICreateTexture2D(
-		VirtualScreen.Width() * CrashTrackerConstants::ScreenScaling,
-		VirtualScreen.Height() * CrashTrackerConstants::ScreenScaling,
+		VirtualScreen.Width(),
+		VirtualScreen.Height(),
 		PF_R8G8B8A8,
 		1,
 		1,
@@ -56,8 +50,8 @@ void FSlateCrashReportResource::InitDynamicRHI()
 	for (int32 i = 0; i < 2; ++i)
 	{
 		ReadbackBuffer[i] = RHICreateTexture2D(
-			VirtualScreen.Width() * CrashTrackerConstants::ScreenScaling,
-			VirtualScreen.Height() * CrashTrackerConstants::ScreenScaling,
+			VirtualScreen.Width(),
+			VirtualScreen.Height(),
 			PF_R8G8B8A8,
 			1,
 			1,
@@ -182,20 +176,6 @@ void FSlateRHIRenderer::Initialize()
 	RenderingPolicy = MakeShareable( new FSlateRHIRenderingPolicy( FontCache, ResourceManager.ToSharedRef() ) ); 
 
 	ElementBatcher = MakeShareable( new FSlateElementBatcher( RenderingPolicy.ToSharedRef() ) );
-	
-#if PLATFORM_WINDOWS || PLATFORM_MAC
-	if (GIsEditor)
-	{
-		FDisplayMetrics DisplayMetrics;
-		FSlateApplication::Get().GetDisplayMetrics(DisplayMetrics);
-		const FIntPoint VirtualScreenOrigin = FIntPoint(DisplayMetrics.VirtualDisplayRect.Left, DisplayMetrics.VirtualDisplayRect.Top);
-		const FIntPoint VirtualScreenLowerRight = FIntPoint(DisplayMetrics.VirtualDisplayRect.Right, DisplayMetrics.VirtualDisplayRect.Bottom);
-		const FIntRect VirtualScreen = FIntRect(VirtualScreenOrigin, VirtualScreenLowerRight);
-	
-		CrashTrackerResource = new FSlateCrashReportResource(VirtualScreen);
-		BeginInitResource(CrashTrackerResource);
-	}
-#endif
 }
 
 void FSlateRHIRenderer::Destroy()
@@ -210,7 +190,7 @@ void FSlateRHIRenderer::Destroy()
 	}
 	
 #if PLATFORM_WINDOWS || PLATFORM_MAC
-	if (GIsEditor)
+	if (CrashTrackerResource != nullptr)
 	{
 		BeginReleaseResource(CrashTrackerResource);
 	}
@@ -230,7 +210,7 @@ void FSlateRHIRenderer::Destroy()
 		delete ViewportInfo;
 	}
 	
-	if (CrashTrackerResource)
+	if (CrashTrackerResource != nullptr)
 	{
 		delete CrashTrackerResource;
 		CrashTrackerResource = NULL;
@@ -621,7 +601,6 @@ void FSlateRHIRenderer::DrawWindows_Private( FSlateDrawBuffer& WindowDrawBuffer
 						FSlateWindowElementList* WindowElementList;
 						SWindow* SlateWindow;
 						bool bLockToVsync;
-						FSimpleDelegate MarkWindowAsDrawn;
 					} Params;
 
 					Params.Renderer = this;
@@ -629,11 +608,6 @@ void FSlateRHIRenderer::DrawWindows_Private( FSlateDrawBuffer& WindowDrawBuffer
 					Params.WindowElementList = &ElementList;
 					Params.bLockToVsync = bLockToVsync;
 
-					if( !Window->HasEverBeenDrawn() )
-					{
-						Params.MarkWindowAsDrawn = Window->MakeMarkWindowAsDrawnDelegate();
-					}
-
 					// NOTE: We pass a raw pointer to the SWindow so that we don't have to use a thread-safe weak pointer in
 					// the FSlateWindowElementList structure
 					Params.SlateWindow = Window.Get();
@@ -642,9 +616,10 @@ void FSlateRHIRenderer::DrawWindows_Private( FSlateDrawBuffer& WindowDrawBuffer
 						FSlateDrawWindowCommandParams, Params, Params,
 					{
 						Params.Renderer->DrawWindow_RenderThread( *Params.ViewportInfo, *Params.WindowElementList, Params.bLockToVsync );
-						Params.MarkWindowAsDrawn.ExecuteIfBound();
 					});
 
+					SlateWindowRendered.Broadcast( *Params.SlateWindow, &ViewInfo->ViewportRHI );
+
 					if ( bTakingAScreenShot )
 					{
 						ENQUEUE_UNIQUE_RENDER_COMMAND_THREEPARAMETER(SlateCaptureScreenshotCommand,
@@ -683,7 +658,84 @@ void FSlateRHIRenderer::DrawWindows_Private( FSlateDrawBuffer& WindowDrawBuffer
 	ElementBatcher->ResetStats();
 }
 
-void FSlateRHIRenderer::CopyWindowsToDrawBuffer(const TArray<FString>& KeypressBuffer)
+
+FIntRect FSlateRHIRenderer::SetupVirtualScreenBuffer(const bool bPrimaryWorkAreaOnly, const float ScreenScaling, ILiveStreamingService* LiveStreamingService )
+{
+#if !PLATFORM_WINDOWS && !PLATFORM_MAC
+	ensureMsg(0, TEXT("This functionality is not valid for this platform"));
+	return;
+#endif
+
+	// Figure out how big we need our render targets to be, based on the size of the entire desktop and the configured scaling amount
+	FDisplayMetrics DisplayMetrics;
+	FSlateApplication::Get().GetDisplayMetrics(DisplayMetrics);
+
+	FIntPoint UnscaledVirtualScreenOrigin;
+	FIntPoint UnscaledVirtualScreenLowerRight;
+
+	if( bPrimaryWorkAreaOnly )
+	{
+		UnscaledVirtualScreenOrigin = FIntPoint(DisplayMetrics.PrimaryDisplayWorkAreaRect.Left, DisplayMetrics.PrimaryDisplayWorkAreaRect.Top);
+		UnscaledVirtualScreenLowerRight = FIntPoint(DisplayMetrics.PrimaryDisplayWorkAreaRect.Right, DisplayMetrics.PrimaryDisplayWorkAreaRect.Bottom);
+	}
+	else
+	{
+		UnscaledVirtualScreenOrigin = FIntPoint(DisplayMetrics.VirtualDisplayRect.Left, DisplayMetrics.VirtualDisplayRect.Top);
+		UnscaledVirtualScreenLowerRight = FIntPoint(DisplayMetrics.VirtualDisplayRect.Right, DisplayMetrics.VirtualDisplayRect.Bottom);
+	}
+	const FIntRect UnscaledVirtualScreen = FIntRect(UnscaledVirtualScreenOrigin, UnscaledVirtualScreenLowerRight);
+	FIntRect ScaledVirtualScreen;
+	{
+		int BufferWidth = FMath::FloorToInt( (float)UnscaledVirtualScreen.Width() * ScreenScaling );
+		int BufferHeight = FMath::FloorToInt( (float)UnscaledVirtualScreen.Height() * ScreenScaling );
+		
+		// If we're preparing a buffer for live streaming, then go ahead and make sure the resolution will be valid for that
+		if( LiveStreamingService != nullptr )
+		{
+			// @todo livestream: This could cause the aspect ratio to be changed and the buffer to be stretched non-uniformly, but usually the aspect only changes slightly
+			LiveStreamingService->MakeValidVideoBufferResolution( BufferWidth, BufferHeight );
+		}
+
+		const float XScaling = (float)BufferWidth / (float)UnscaledVirtualScreen.Width();
+		const float YScaling = (float)BufferHeight / (float)UnscaledVirtualScreen.Height();
+
+		ScaledVirtualScreen.Min.X = FMath::FloorToInt(( float)UnscaledVirtualScreen.Min.X * XScaling );
+		ScaledVirtualScreen.Max.X = FMath::FloorToInt(( float)UnscaledVirtualScreen.Max.X * XScaling );
+		ScaledVirtualScreen.Min.Y = FMath::FloorToInt(( float)UnscaledVirtualScreen.Min.Y * YScaling );
+		ScaledVirtualScreen.Max.Y = FMath::FloorToInt(( float)UnscaledVirtualScreen.Max.Y * YScaling );
+	}
+
+	// @todo livestream: This CrashTrackerResource is now also used for editor live streaming, so we should consider renaming it and cleaning
+	// up the API a little bit more
+	if( CrashTrackerResource == nullptr || 
+		ScaledVirtualScreen != CrashTrackerResource->GetVirtualScreen() )
+	{
+		if( CrashTrackerResource != nullptr )
+		{
+			// Size has changed, so clear out our old resource and create a new one
+			if (CrashTrackerResource != nullptr)
+			{
+				BeginReleaseResource(CrashTrackerResource);
+			}
+
+			FlushRenderingCommands();
+	
+			if (CrashTrackerResource != nullptr)
+			{
+				delete CrashTrackerResource;
+				CrashTrackerResource = NULL;
+			}
+		}
+
+		CrashTrackerResource = new FSlateCrashReportResource(ScaledVirtualScreen, UnscaledVirtualScreen);
+		BeginInitResource(CrashTrackerResource);
+	}
+
+	return CrashTrackerResource->GetVirtualScreen();
+}
+
+
+void FSlateRHIRenderer::CopyWindowsToVirtualScreenBuffer(const TArray<FString>& KeypressBuffer)
 {
 #if !PLATFORM_WINDOWS && !PLATFORM_MAC
 	ensureMsg(0, TEXT("This functionality is not valid for this platform"));
@@ -692,9 +744,15 @@ void FSlateRHIRenderer::CopyWindowsToDrawBuffer(const TArray<FString>& KeypressB
 
 	SCOPE_CYCLE_COUNTER(STAT_GenerateCaptureBuffer);
 
-	const FIntRect ScaledVirtualScreen = CrashTrackerResource->GetVirtualScreen().Scale(CrashTrackerConstants::ScreenScaling);
-	const FIntPoint ScaledVirtualScreenPos = ScaledVirtualScreen.Min;
-	const FIntPoint ScaledVirtualScreenSize = ScaledVirtualScreen.Size();
+	// Make sure to call SetupDrawBuffer() before calling this function!
+	check( CrashTrackerResource != nullptr );
+
+	const FIntRect VirtualScreen = CrashTrackerResource->GetVirtualScreen();
+	const FIntPoint VirtualScreenPos = VirtualScreen.Min;
+	const FIntPoint VirtualScreenSize = VirtualScreen.Size();
+	const FIntRect UnscaledVirtualScreen = CrashTrackerResource->GetUnscaledVirtualScreen();
+	const float XScaling = (float)VirtualScreen.Width() / (float)UnscaledVirtualScreen.Width();
+	const float YScaling = (float)VirtualScreen.Height() / (float)UnscaledVirtualScreen.Height();
 	
 	// setup state
 	struct FSetupWindowStateContext
@@ -705,7 +763,7 @@ void FSlateRHIRenderer::CopyWindowsToDrawBuffer(const TArray<FString>& KeypressB
 	FSetupWindowStateContext SetupWindowStateContext =
 	{
 		CrashTrackerResource,
-		ScaledVirtualScreen
+		VirtualScreen
 	};
 
 	ENQUEUE_UNIQUE_RENDER_COMMAND_ONEPARAMETER(
@@ -719,13 +777,17 @@ void FSlateRHIRenderer::CopyWindowsToDrawBuffer(const TArray<FString>& KeypressB
 		RHISetRasterizerState(TStaticRasterizerState<>::GetRHI());
 		RHISetDepthStencilState(TStaticDepthStencilState<false,CF_Always>::GetRHI());
 
-		RHIClear(true, FLinearColor::Gray, false, 0.f, false, 0x00, FIntRect());
+		// @todo livestream: Ideally this "desktop background color" should be configurable in the editor's preferences
+		RHIClear(true, FLinearColor( 0.02f, 0.02f, 0.2f ), false, 0.f, false, 0x00, FIntRect());
 	});
 
 	// draw windows to buffer
 	TArray< TSharedRef<SWindow> > OutWindows;
 	FSlateApplication::Get().GetAllVisibleWindowsOrdered(OutWindows);
 
+	static const FName RendererModuleName( "Renderer" );
+	IRendererModule& RendererModule = FModuleManager::GetModuleChecked<IRendererModule>( RendererModuleName );
+
 	for (int32 i = 0; i < OutWindows.Num(); ++i)
 	{
 		TSharedPtr<SWindow> WindowPtr = OutWindows[i];
@@ -736,56 +798,84 @@ void FSlateRHIRenderer::CopyWindowsToDrawBuffer(const TArray<FString>& KeypressB
 		const FVector2D WindowSize = SlateWindowRect.GetSize();
 		if ( WindowSize.X > 0 && WindowSize.Y > 0 )
 		{
-			const FIntRect ScaledWindowRect = FIntRect(SlateWindowRect.Left, SlateWindowRect.Top, SlateWindowRect.Right, SlateWindowRect.Bottom).Scale(CrashTrackerConstants::ScreenScaling) -
-				ScaledVirtualScreenPos;
+			FIntRect ScaledWindowRect;
+			ScaledWindowRect.Min.X = SlateWindowRect.Left * XScaling - VirtualScreenPos.X;
+			ScaledWindowRect.Max.X = SlateWindowRect.Right * XScaling - VirtualScreenPos.X;
+			ScaledWindowRect.Min.Y = SlateWindowRect.Top * YScaling - VirtualScreenPos.Y;
+			ScaledWindowRect.Max.Y = SlateWindowRect.Bottom * YScaling - VirtualScreenPos.Y;
 
 			struct FDrawWindowToBufferContext
 			{
 				FViewportInfo* InViewportInfo;
 				FIntRect WindowRect;
 				FIntRect IntermediateBufferSize;
+				IRendererModule* RendererModule;
 			};
 			FDrawWindowToBufferContext DrawWindowToBufferContext =
 			{
 				ViewportInfo,
 				ScaledWindowRect,
-				ScaledVirtualScreen
+				VirtualScreen,
+				&RendererModule
 			};
 
+			// Draw a quad mapping scene color to the view's render target
 			ENQUEUE_UNIQUE_RENDER_COMMAND_ONEPARAMETER(
 				DrawWindowToBuffer,
 				FDrawWindowToBufferContext,Context,DrawWindowToBufferContext,
 			{
-				//@todo-rco: RHIPacketList
+				TShaderMapRef<FScreenVS> VertexShader(GetGlobalShaderMap());
+				TShaderMapRef<FScreenPS> PixelShader(GetGlobalShaderMap());
+
 				FRHICommandList& RHICmdList = FRHICommandList::GetNullRef();
 
-				DrawNormalizedScreenQuad(
+				static FGlobalBoundShaderState BoundShaderState;
+				SetGlobalBoundShaderState(RHICmdList, BoundShaderState, Context.RendererModule->GetFilterVertexDeclaration().VertexDeclarationRHI, *VertexShader, *PixelShader);
+
+				if( Context.WindowRect.Width() != Context.InViewportInfo->Width || Context.WindowRect.Height() != Context.InViewportInfo->Height )
+				{
+					// We're scaling down the window, so use bilinear filtering
+					PixelShader->SetParameters(RHICmdList, TStaticSamplerState<SF_Bilinear>::GetRHI(), RHIGetViewportBackBuffer(Context.InViewportInfo->ViewportRHI));
+				}
+				else
+				{
+					// Drawing 1:1, so no filtering needed
+					PixelShader->SetParameters(RHICmdList, TStaticSamplerState<SF_Point>::GetRHI(), RHIGetViewportBackBuffer(Context.InViewportInfo->ViewportRHI));
+				}
+
+				Context.RendererModule->DrawRectangle(
 					RHICmdList,
 					Context.WindowRect.Min.X, Context.WindowRect.Min.Y,
-					0, 0,
 					Context.WindowRect.Width(), Context.WindowRect.Height(),
+					0, 0,
 					1, 1,
 					FIntPoint(Context.IntermediateBufferSize.Width(), Context.IntermediateBufferSize.Height()),
-					RHIGetViewportBackBuffer(Context.InViewportInfo->ViewportRHI));
+					FIntPoint(1, 1),
+					*VertexShader,
+					EDRF_Default);
 			});
 		}
 	}
 
 	// draw mouse cursor and keypresses
 	const FVector2D MouseCursorLocation = FSlateApplication::Get().GetCursorPos();
-	const FIntPoint ScaledCursorLocation = FIntPoint(MouseCursorLocation.X * CrashTrackerConstants::ScreenScaling, MouseCursorLocation.Y * CrashTrackerConstants::ScreenScaling) -
-		ScaledVirtualScreenPos;
+	const FIntPoint ScaledCursorLocation = FIntPoint(MouseCursorLocation.X * XScaling, MouseCursorLocation.Y * YScaling) - VirtualScreenPos;
 
 	FSlateWindowElementList* WindowElementList = CrashTrackerResource->GetNextElementList();
 	*WindowElementList = FSlateWindowElementList(TSharedPtr<SWindow>());
 
-	FSlateDrawElement::MakeBox(
-		*WindowElementList,
-		0,
-		FPaintGeometry(ScaledCursorLocation, FVector2D(32, 32), 1.f),
-		FCoreStyle::Get().GetBrush("CrashTracker.Cursor"),
-		FSlateRect(0, 0, ScaledVirtualScreenSize.X, ScaledVirtualScreenSize.Y));
-	
+	// Don't draw cursor when it is hidden (mouse looking, scrolling, etc.)
+	// @todo livestream: The cursor is probably still hidden when dragging with the mouse captured (grabby hand)
+	if( FSlateApplication::Get().GetMouseCaptureWindow() == nullptr )
+	{
+		FSlateDrawElement::MakeBox(
+			*WindowElementList,
+			0,
+			FPaintGeometry(ScaledCursorLocation, FVector2D(32, 32) * XScaling, XScaling),
+			FCoreStyle::Get().GetBrush("CrashTracker.Cursor"),
+			FSlateRect(0, 0, VirtualScreenSize.X, VirtualScreenSize.Y));
+	}
+
 	for (int32 i = 0; i < KeypressBuffer.Num(); ++i)
 	{
 		FSlateDrawElement::MakeText(
@@ -794,7 +884,7 @@ void FSlateRHIRenderer::CopyWindowsToDrawBuffer(const TArray<FString>& KeypressB
 			FPaintGeometry(FVector2D(10, 10 + i * 30), FVector2D(300, 30), 1.f),
 			KeypressBuffer[i],
 			FCoreStyle::Get().GetFontStyle(TEXT("CrashTracker.Font")),
-			FSlateRect(0, 0, ScaledVirtualScreenSize.X, ScaledVirtualScreenSize.Y));
+			FSlateRect(0, 0, VirtualScreenSize.X, VirtualScreenSize.Y));
 	}
 	
 	ElementBatcher->AddElements(WindowElementList->GetDrawElements());
@@ -814,10 +904,10 @@ void FSlateRHIRenderer::CopyWindowsToDrawBuffer(const TArray<FString>& KeypressB
 	FWriteMouseCursorAndKeyPressesContext WriteMouseCursorAndKeyPressesContext =
 	{
 		CrashTrackerResource,
-		ScaledVirtualScreen,
+		VirtualScreen,
 		RenderingPolicy.Get(),
 		WindowElementList,
-		ScaledVirtualScreenSize
+		VirtualScreenSize
 	};
 	ENQUEUE_UNIQUE_RENDER_COMMAND_ONEPARAMETER(
 		WriteMouseCursorAndKeyPresses,
@@ -844,7 +934,7 @@ void FSlateRHIRenderer::CopyWindowsToDrawBuffer(const TArray<FString>& KeypressB
 	FReadbackFromIntermediateBufferContext ReadbackFromIntermediateBufferContext =
 	{
 		CrashTrackerResource,
-		ScaledVirtualScreen
+		VirtualScreen
 	};
 	ENQUEUE_UNIQUE_RENDER_COMMAND_ONEPARAMETER(
 		ReadbackFromIntermediateBuffer,
@@ -859,33 +949,31 @@ void FSlateRHIRenderer::CopyWindowsToDrawBuffer(const TArray<FString>& KeypressB
 }
 
 
-void FSlateRHIRenderer::MapCrashTrackerBuffer(void** OutImageData, int32* OutWidth, int32* OutHeight)
+void FSlateRHIRenderer::MapVirtualScreenBuffer(void** OutImageData)
 {
 	struct FReadbackFromStagingBufferContext
 	{
 		FSlateCrashReportResource* CrashReportResource;
 		void** OutData;
-		int32* OutputWidth;
-		int32* OutputHeight;
 	};
 	FReadbackFromStagingBufferContext ReadbackFromStagingBufferContext =
 	{
 		CrashTrackerResource,
-		OutImageData,
-		OutWidth,
-		OutHeight
+		OutImageData
 	};
 	ENQUEUE_UNIQUE_RENDER_COMMAND_ONEPARAMETER(
 		ReadbackFromStagingBuffer,
 		FReadbackFromStagingBufferContext,Context,ReadbackFromStagingBufferContext,
 	{
 		SCOPE_CYCLE_COUNTER(STAT_MapStagingBuffer);
-		RHIMapStagingSurface(Context.CrashReportResource->GetReadbackBuffer(), *Context.OutData, *Context.OutputWidth, *Context.OutputHeight);
+		int32 UnusedWidth = 0;
+		int32 UnusedHeight = 0;
+		RHIMapStagingSurface(Context.CrashReportResource->GetReadbackBuffer(), *Context.OutData, UnusedWidth, UnusedHeight);
 		Context.CrashReportResource->SwapTargetReadbackBuffer();
 	});
 }
 
-void FSlateRHIRenderer::UnmapCrashTrackerBuffer()
+void FSlateRHIRenderer::UnmapVirtualScreenBuffer()
 {
 	struct FReadbackFromStagingBufferContext
 	{
diff --git a/Engine/Source/Runtime/SlateRHIRenderer/Private/SlateRHIRenderer.h b/Engine/Source/Runtime/SlateRHIRenderer/Private/SlateRHIRenderer.h
--- a/Engine/Source/Runtime/SlateRHIRenderer/Private/SlateRHIRenderer.h
+++ b/Engine/Source/Runtime/SlateRHIRenderer/Private/SlateRHIRenderer.h
@@ -33,10 +33,11 @@ private:
 class FSlateCrashReportResource : public FRenderResource
 {
 public:
-	FSlateCrashReportResource(FIntRect InVirtualScreen)
+	FSlateCrashReportResource(FIntRect InVirtualScreen, FIntRect InUnscaledVirtualScreen)
 		: ReadbackBufferIndex(0)
 		, ElementListIndex(0)
-		, VirtualScreen(InVirtualScreen) {}
+		, VirtualScreen(InVirtualScreen)
+		, UnscaledVirtualScreen(InUnscaledVirtualScreen) {}
 
 	/** FRenderResource Interface.  Called when render resources need to be initialized */
 	virtual void InitDynamicRHI();
@@ -52,6 +53,7 @@ public:
 
 	/** Accessors */
 	FIntRect GetVirtualScreen() const {return VirtualScreen;}
+	FIntRect GetUnscaledVirtualScreen() const {return UnscaledVirtualScreen;}
 	FTexture2DRHIRef GetBuffer() const {return CrashReportBuffer;}
 	FTexture2DRHIRef GetReadbackBuffer() const {return ReadbackBuffer[ReadbackBufferIndex];}
 
@@ -71,6 +73,9 @@ private:
 
 	/** The size of the virtual screen, used to calculate the buffer size */
 	FIntRect VirtualScreen;
+
+	/** Size of the virtual screen before we applied any scaling */
+	FIntRect UnscaledVirtualScreen;
 };
 
 
@@ -158,15 +163,26 @@ public:
 	void DrawWindow_RenderThread( const FSlateRHIRenderer::FViewportInfo& ViewportInfo, const FSlateWindowElementList& WindowElementList, bool bLockToVsync );
 
 
+	/**
+	 * You must call this before calling CopyWindowsToVirtualScreenBuffer(), to setup the render targets first.
+	 * 
+	 * @param	ScreenScaling	How much to downscale the desktop size
+	 * @param	LiveStreamingService	Optional pointer to a live streaming service this buffer needs to work with
+	 * @param	bPrimaryWorkAreaOnly	True if we should capture only the primary monitor's work area, or false to capture the entire desktop spanning all monitors
+	 *
+	 * @return	The virtual screen rectangle.  The size of this rectangle will be the size of the render target buffer.
+	 */
+	virtual FIntRect SetupVirtualScreenBuffer(const bool bPrimaryWorkAreaOnly, const float ScreenScaling, class ILiveStreamingService* LiveStreamingService) override;
+
 	/**
 	 * Copies all slate windows out to a buffer at half resolution with debug information
 	 * like the mouse cursor and any keypresses.
 	 */
-	virtual void CopyWindowsToDrawBuffer(const TArray<FString>& KeypressBuffer);
+	virtual void CopyWindowsToVirtualScreenBuffer(const TArray<FString>& KeypressBuffer) override;
 	
 	/** Allows and disallows access to the crash tracker buffer data on the CPU */
-	virtual void MapCrashTrackerBuffer(void** OutImageData, int32* OutWidth, int32* OutHeight) override;
-	virtual void UnmapCrashTrackerBuffer() override;
+	virtual void MapVirtualScreenBuffer(void** OutImageData) override;
+	virtual void UnmapVirtualScreenBuffer() override;
 
 	/**
 	 * Reloads texture resources from disk                   
