Fix non-windows/mac compiles.

diff --git a/Engine/Source/Runtime/SlateRHIRenderer/Private/SlateRHIRenderer.cpp b/Engine/Source/Runtime/SlateRHIRenderer/Private/SlateRHIRenderer.cpp
--- a/Engine/Source/Runtime/SlateRHIRenderer/Private/SlateRHIRenderer.cpp
+++ b/Engine/Source/Runtime/SlateRHIRenderer/Private/SlateRHIRenderer.cpp
@@ -661,15 +661,15 @@ void FSlateRHIRenderer::DrawWindows_Private( FSlateDrawBuffer& WindowDrawBuffer
 
 FIntRect FSlateRHIRenderer::SetupVirtualScreenBuffer(const bool bPrimaryWorkAreaOnly, const float ScreenScaling, ILiveStreamingService* LiveStreamingService )
 {
-#if !PLATFORM_WINDOWS && !PLATFORM_MAC
-	ensureMsg(0, TEXT("This functionality is not valid for this platform"));
-	return;
-#endif
-
 	// Figure out how big we need our render targets to be, based on the size of the entire desktop and the configured scaling amount
 	FDisplayMetrics DisplayMetrics;
 	FSlateApplication::Get().GetDisplayMetrics(DisplayMetrics);
 
+#if !PLATFORM_WINDOWS && !PLATFORM_MAC
+	ensureMsg(0, TEXT("This functionality is not valid for this platform"));	
+	return FIntRect(FIntPoint(0, 0), FIntPoint(DisplayMetrics.PrimaryDisplayWidth, DisplayMetrics.PrimaryDisplayHeight));
+#endif	
+
 	FIntPoint UnscaledVirtualScreenOrigin;
 	FIntPoint UnscaledVirtualScreenLowerRight;
 
