UE4: Fixed a crash caused by using the same texture in both the legacy native brushes and the new ones as well.

diff --git a/Engine/Source/Runtime/SlateRHIRenderer/Private/SlateRHIResourceManager.cpp b/Engine/Source/Runtime/SlateRHIRenderer/Private/SlateRHIResourceManager.cpp
--- a/Engine/Source/Runtime/SlateRHIRenderer/Private/SlateRHIResourceManager.cpp
+++ b/Engine/Source/Runtime/SlateRHIRenderer/Private/SlateRHIResourceManager.cpp
@@ -56,7 +56,16 @@ TSharedPtr<FDynamicTextureResource> FDynamicResourceMap::GetTextureResource( FNa
 {
 	if( TextureObject )
 	{
-		return DynamicResourceObjectMap.FindRef( TextureObject );
+		const TSharedPtr<FDynamicTextureResource>& FoundResource = DynamicResourceObjectMap.FindRef( TextureObject );
+		if ( FoundResource.IsValid() )
+		{
+			return FoundResource;
+		}
+		else
+		{
+			// We might have already created a resource for this texture in the Native map. Look that up here.
+			return DynamicNativeTextureMapLookupByTextureObject.FindRef( TextureObject ).Pin();
+		}
 	}
 	else
 	{
@@ -81,6 +90,11 @@ void FDynamicResourceMap::AddTextureResource( FName ResourceName, UTexture2D* Te
 	else
 	{
 		DynamicNativeTextureMap.Add( ResourceName, InResource );
+
+		if ( InResource->TextureObject )
+		{
+			DynamicNativeTextureMapLookupByTextureObject.Add(InResource->TextureObject, InResource);
+		}
 	}
 
 }
@@ -99,7 +113,13 @@ void FDynamicResourceMap::RemoveTextureResource( FName ResourceName, UTexture2D*
 	}
 	else
 	{
+		const TSharedPtr<FDynamicTextureResource>& ExistingResource = DynamicNativeTextureMap.FindRef( ResourceName );
 		DynamicNativeTextureMap.Remove( ResourceName );
+
+		if (ExistingResource.IsValid() && ExistingResource->TextureObject)
+		{
+			DynamicNativeTextureMapLookupByTextureObject.Remove(ExistingResource->TextureObject);
+		}
 	}
 }
 
@@ -111,6 +131,7 @@ void FDynamicResourceMap::RemoveMaterialResource( UMaterialInterface* Material )
 void FDynamicResourceMap::Empty()
 {
 	DynamicNativeTextureMap.Empty();
+	DynamicNativeTextureMapLookupByTextureObject.Empty();
 	DynamicResourceObjectMap.Empty();
 	MaterialResourceMap.Empty();
 }
diff --git a/Engine/Source/Runtime/SlateRHIRenderer/Private/SlateRHIResourceManager.h b/Engine/Source/Runtime/SlateRHIRenderer/Private/SlateRHIResourceManager.h
--- a/Engine/Source/Runtime/SlateRHIRenderer/Private/SlateRHIResourceManager.h
+++ b/Engine/Source/Runtime/SlateRHIRenderer/Private/SlateRHIResourceManager.h
@@ -55,6 +55,8 @@ public:
 
 private:
 	TMap<FName, TSharedPtr<FDynamicTextureResource> > DynamicNativeTextureMap;
+	
+	TMap<UObject*, TWeakPtr<FDynamicTextureResource> > DynamicNativeTextureMapLookupByTextureObject;
 
 	TMap<UObject*, TSharedPtr<FDynamicTextureResource> > DynamicResourceObjectMap;
 
