Added material support to Slate (unlit only)
- Materials are usable on any brush type and there is a new dynamic brush type (FSlateMaterialBrush) for MIDs

diff --git a/Engine/Shaders/SlateElementPixelShader.usf b/Engine/Shaders/SlateElementPixelShader.usf
--- a/Engine/Shaders/SlateElementPixelShader.usf
+++ b/Engine/Shaders/SlateElementPixelShader.usf
@@ -27,7 +27,7 @@
 #define DEBUG_CLIPPING 0
 
 /** Display gamma */
-float DisplayGamma;
+float InvDisplayGamma;
 /** Per element params */
 float4 ShaderParams;
 
@@ -51,14 +51,20 @@ struct VertexOut
 float4 GetMaterialColor( VertexOut InVertex, float2 UV )
 {
 	FMaterialPixelParameters Parameters = MakeInitializedMaterialPixelParameters();
-	CalcMaterialParameters(Parameters, true, InVertex.Position);
-#if NUM_MATERIAL_TEXCOORDS
-	for(int CoordinateIndex = 0;CoordinateIndex < NUM_MATERIAL_TEXCOORDS;CoordinateIndex++)
-	{
-		Parameters.TexCoords[CoordinateIndex] = InVertex.TextureCoordinates.xy*InVertex.TextureCoordinates.zw;
-	}
+	CalcMaterialParameters(Parameters, true, float4(InVertex.WorldPosition,0,1));
+#if NUM_MATERIAL_TEXCOORDS == 1
+	Parameters.TexCoords[0] = UV;
+#elif NUM_MATERIAL_TEXCOORDS == 2
+	Parameters.TexCoords[0] = UV;
+	Parameters.TexCoords[1] = InVertex.TextureCoordinates.xy;
+#elif NUM_MATERIAL_TEXCOORDS == 3
+	Parameters.TexCoords[0] = UV;
+	Parameters.TexCoords[1] = InVertex.TextureCoordinates.xy;
+	Parameters.TexCoords[2] = InVertex.TextureCoordinates.zw;
 #endif
 
+	Parameters.VertexColor = InVertex.Color;
+
 	// Clip if masked
 	GetMaterialCoverageAndClipping( Parameters );
 
@@ -278,11 +284,13 @@ float4 Main( VertexOut VIn ) : SV_Target0
 	OutColor.rgb = lerp( OutColor.rgb, Grayish, clamp( distance( OutColor.rgb, Grayish ), 0, .8)  );
 #endif
 
+#if !USE_MATERIALS
 	OutColor = ApplyColorVisionDeficiencyFilter(OutColor);
+#endif
 
 #if !COMPILER_GLSL_ES2
 	// gamma correct
-	OutColor.rgb = pow(OutColor.rgb,1.0f/DisplayGamma);
+	OutColor.rgb = pow(OutColor.rgb,InvDisplayGamma);
 #endif
 
 
diff --git a/Engine/Source/Editor/DetailCustomizations/Private/SlateBrushCustomization.cpp b/Engine/Source/Editor/DetailCustomizations/Private/SlateBrushCustomization.cpp
--- a/Engine/Source/Editor/DetailCustomizations/Private/SlateBrushCustomization.cpp
+++ b/Engine/Source/Editor/DetailCustomizations/Private/SlateBrushCustomization.cpp
@@ -181,6 +181,7 @@ public:
 			]
 		);
 
+		CachedTextureSize = FVector2D::ZeroVector;
 		CachePropertyValues();
 		SetDefaultAlignment();
 		UpdatePreviewImageSize();
@@ -575,15 +576,23 @@ private:
 		UObject* ResourceObject;
 		FPropertyAccess::Result Result = ResourceObjectProperty->GetValue( ResourceObject );
 		if( Result == FPropertyAccess::Success )
-		{
-			UTexture2D* BrushTexture = Cast<UTexture2D>(ResourceObject);
-			CachedTextureSize = BrushTexture ? FVector2D( BrushTexture->GetSizeX(), BrushTexture->GetSizeY() ) : FVector2D( 32.0f, 32.0f );
-
+		{				
 			TArray<void*> RawData;
-			ImageSizeProperty->AccessRawData( RawData );
-			if( RawData.Num() > 0 && RawData[ 0 ] != NULL )
+			ImageSizeProperty->AccessRawData(RawData);
+			if (RawData.Num() > 0 && RawData[0] != NULL)
+			{
+				CachedImageSizeValue = *static_cast<FVector2D*>(RawData[0]);
+			}
+
+			UTexture2D* BrushTexture = Cast<UTexture2D>(ResourceObject);
+			if( BrushTexture )
+			{
+				CachedTextureSize = FVector2D( BrushTexture->GetSizeX(), BrushTexture->GetSizeY() );
+			}
+			else if( CachedTextureSize == FVector2D::ZeroVector )
 			{
-				CachedImageSizeValue = *static_cast<FVector2D*>(RawData[ 0 ]);
+				// If the cached texture size is not initialized, create a default value now for materials 
+				CachedTextureSize = CachedImageSizeValue;
 			}
 
 			uint8 DrawAsType;
@@ -933,6 +942,71 @@ private:
 const float SSlateBrushPreview::ImagePadding = 5.0f;
 const float SSlateBrushPreview::BorderHitSize = 8.0f;
 
+class SBrushResourceObjectBox : public SCompoundWidget
+{
+	SLATE_BEGIN_ARGS( SBrushResourceObjectBox ) {}
+	
+	SLATE_END_ARGS()
+
+	void Construct( const FArguments& InArgs, IStructCustomizationUtils* StructCustomizationUtils, TSharedPtr<IPropertyHandle> InResourceObjectProperty )
+	{
+		ResourceObjectProperty = InResourceObjectProperty;
+		ChildSlot
+		[
+			SNew(SVerticalBox)
+			+ SVerticalBox::Slot()
+			.FillHeight(1)
+			[
+				SNew(SObjectPropertyEntryBox)
+				.PropertyHandle(InResourceObjectProperty)
+				.ThumbnailPool(StructCustomizationUtils->GetThumbnailPool())
+				.OnShouldFilterAsset(this, &SBrushResourceObjectBox::OnFilterAssetPicker)
+			]
+			+ SVerticalBox::Slot()
+			.AutoHeight()
+			.Padding( 0.0f, 3.0f )
+			[
+				SAssignNew(ResourceErrorText, SErrorText)
+			]
+		];
+	}
+
+	void Tick( const FGeometry& AllottedGeometry, const double InCurrentTime, const float InDeltaTime )
+	{
+		UObject* Resource = nullptr;
+
+		if( ResourceObjectProperty->GetValue(Resource) == FPropertyAccess::Success && Resource && Resource->IsA<UMaterialInterface>() )
+		{
+			UMaterialInterface* MaterialInterface = Cast<UMaterialInterface>( Resource );
+			UMaterial* BaseMaterial = MaterialInterface->GetBaseMaterial();
+			if( BaseMaterial && !BaseMaterial->bUsedWithUI )
+			{
+				ResourceErrorText->SetError( NSLOCTEXT("FSlateBrushStructCustomization", "ResourceErrorText", "This material is not supported in UI.  Please check \"Used with UI\" in the material editor" ) );
+			}
+			else
+			{
+				ResourceErrorText->SetError( FText::GetEmpty() );
+			}
+		}
+		else if( ResourceErrorText->HasError() )
+		{
+			ResourceErrorText->SetError( FText::GetEmpty() );
+		}
+	}
+
+private:
+
+	/** Called when the asset picker needs to be filtered */
+	bool OnFilterAssetPicker(const FAssetData& InAssetData) const
+	{
+		UClass* Class = InAssetData.GetClass();
+		return !Class->IsChildOf(UTexture2D::StaticClass()) && !Class->IsChildOf(UMaterialInterface::StaticClass());
+	}
+
+private:
+	TSharedPtr<IPropertyHandle> ResourceObjectProperty;
+	TSharedPtr<SErrorText> ResourceErrorText;
+};
 
 TSharedRef<IPropertyTypeCustomization> FSlateBrushStructCustomization::MakeInstance() 
 {
@@ -972,10 +1046,7 @@ void FSlateBrushStructCustomization::CustomizeChildren( TSharedRef<IPropertyHand
 	.MinDesiredWidth(250.0f)
 	.MaxDesiredWidth(0.0f)
 	[
-		SNew( SObjectPropertyEntryBox )
-		.PropertyHandle( ResourceObjectProperty )
-		.ThumbnailPool( StructCustomizationUtils.GetThumbnailPool() )
-		.OnShouldFilterAsset( this, &FSlateBrushStructCustomization::OnFilterAssetPicker )
+		SNew( SBrushResourceObjectBox, &StructCustomizationUtils, ResourceObjectProperty )
 	];
 
 	StructBuilder.AddChildProperty( ImageSizeProperty.ToSharedRef() );
@@ -1029,11 +1100,6 @@ void FSlateBrushStructCustomization::CustomizeChildren( TSharedRef<IPropertyHand
 	}
 }
 
-bool FSlateBrushStructCustomization::OnFilterAssetPicker( const FAssetData& InAssetData ) const
-{
-	UClass* Class = InAssetData.GetClass();
-	return !Class->IsChildOf( UTexture2D::StaticClass() );//&& !Class->IsChildOf( UMaterialInterface::StaticClass() );
-}
 
 EVisibility FSlateBrushStructCustomization::GetTilingPropertyVisibility() const
 {
diff --git a/Engine/Source/Editor/DetailCustomizations/Private/SlateBrushCustomization.h b/Engine/Source/Editor/DetailCustomizations/Private/SlateBrushCustomization.h
--- a/Engine/Source/Editor/DetailCustomizations/Private/SlateBrushCustomization.h
+++ b/Engine/Source/Editor/DetailCustomizations/Private/SlateBrushCustomization.h
@@ -23,11 +23,12 @@ private:
 	 *  Get the Slate Brush margin property row visibility
 	 */
 	EVisibility GetMarginPropertyVisibility() const;
-	
-	/** Called when the asset picker needs to be filtered */
-	bool OnFilterAssetPicker( const FAssetData& InAssetData ) const;
+
 
 	/** Slate Brush DrawAs property */
 	TSharedPtr<IPropertyHandle> DrawAsProperty;
+
+	/** Error text to display if the resource object is not valid*/
+	TSharedPtr<SErrorText> ResourceErrorText;
 };
 
diff --git a/Engine/Source/Runtime/Slate/Public/Slate.h b/Engine/Source/Runtime/Slate/Public/Slate.h
--- a/Engine/Source/Runtime/Slate/Public/Slate.h
+++ b/Engine/Source/Runtime/Slate/Public/Slate.h
@@ -17,8 +17,6 @@
 #define BEGIN_SLATE_FUNCTION_BUILD_OPTIMIZATION PRAGMA_DISABLE_OPTIMIZATION
 #define END_SLATE_FUNCTION_BUILD_OPTIMIZATION   PRAGMA_ENABLE_OPTIMIZATION
 
-#define ALLOW_SLATE_MATERIALS 0
-
 /** Stat tracking for Slate */
 DECLARE_CYCLE_STAT_EXTERN( TEXT("Total Slate Tick Time"), STAT_SlateTickTime, STATGROUP_Slate , );
 DECLARE_CYCLE_STAT_EXTERN( TEXT("CacheDesiredSize"), STAT_SlateCacheDesiredSize, STATGROUP_Slate , );
diff --git a/Engine/Source/Runtime/Slate/Public/SlateMaterialBrush.h b/Engine/Source/Runtime/Slate/Public/SlateMaterialBrush.h
--- a/Engine/Source/Runtime/Slate/Public/SlateMaterialBrush.h
+++ b/Engine/Source/Runtime/Slate/Public/SlateMaterialBrush.h
@@ -4,8 +4,11 @@
 
 #include "Slate.h"
 
-#if ALLOW_SLATE_MATERIALS
-
+/**
+ * Dynamic rush for referencing a UMaterial.  
+ *
+ * Note: This brush nor the slate renderer holds a strong reference to the material.  You are responsible for maintaining the lifetime of the brush and material object.
+ */
 struct FSlateMaterialBrush : public FSlateBrush
 {
 	FSlateMaterialBrush( class UMaterialInterface& InMaterial, const FVector2D& InImageSize )
@@ -23,5 +26,3 @@ struct FSlateMaterialBrush : public FSlateBrush
 		}
 	}
 }; 
-
-#endif
\ No newline at end of file
diff --git a/Engine/Source/Runtime/SlateCore/Public/Styling/SlateBrush.h b/Engine/Source/Runtime/SlateCore/Public/Styling/SlateBrush.h
--- a/Engine/Source/Runtime/SlateCore/Public/Styling/SlateBrush.h
+++ b/Engine/Source/Runtime/SlateCore/Public/Styling/SlateBrush.h
@@ -203,7 +203,7 @@ public:
 protected:
 
 	/** Pointer to the UTexture2D. Holding onto it as a UObject because UTexture2D is not available in Slate core. */
-	UPROPERTY(EditAnywhere, Category=Brush, meta=(DisplayThumbnail="true", DisplayName="Texture"))
+	UPROPERTY(EditAnywhere, Category=Brush, meta=(DisplayThumbnail="true", DisplayName="Texture or Material Asset"))
 	UObject* ResourceObject;
 
 	/** The name of the rendering resource to use */
diff --git a/Engine/Source/Runtime/SlateRHIRenderer/Private/SlateMaterialShader.cpp b/Engine/Source/Runtime/SlateRHIRenderer/Private/SlateMaterialShader.cpp
new file mode 100644
--- /dev/null
+++ b/Engine/Source/Runtime/SlateRHIRenderer/Private/SlateMaterialShader.cpp
@@ -0,0 +1,81 @@
+// Copyright 1998-2014 Epic Games, Inc. All Rights Reserved.
+
+#include "SlateRHIRendererPrivatePCH.h"
+#include "SlateMaterialShader.h"
+#include "RenderingCommon.h"
+
+bool FSlateMaterialShaderPS::ShouldCache(EShaderPlatform Platform, const FMaterial* Material)
+{
+	return Material->IsUsedWithUI();
+}
+
+
+void FSlateMaterialShaderPS::ModifyCompilationEnvironment(EShaderPlatform Platform, const FMaterial* Material, FShaderCompilerEnvironment& OutEnvironment)
+{
+	// Set defines based on what this shader will be used for
+	OutEnvironment.SetDefine( TEXT("USE_MATERIALS"), 1 );
+
+	FMaterialShader::ModifyCompilationEnvironment( Platform, Material, OutEnvironment );
+}
+
+FSlateMaterialShaderPS::FSlateMaterialShaderPS(const FMaterialShaderType::CompiledShaderInitializerType& Initializer)
+	: FMaterialShader(Initializer)
+{
+	ShaderParams.Bind(Initializer.ParameterMap, TEXT("ShaderParams"));
+	DisplayGamma.Bind(Initializer.ParameterMap, TEXT("InvDisplayGamma"));
+}
+
+void FSlateMaterialShaderPS::SetParameters(FRHICommandList& RHICmdList, const FSceneView& View, const FMaterialRenderProxy* MaterialRenderProxy, const FMaterial* Material, float InDisplayGamma, const FVector4& InShaderParams )
+{
+	const FPixelShaderRHIParamRef ShaderRHI = GetPixelShader();
+
+	EBlendMode BlendMode = Material->GetBlendMode();
+
+	switch (BlendMode)
+	{
+	default:
+	case BLEND_Opaque:
+		RHICmdList.SetBlendState(TStaticBlendState<>::GetRHI());
+		break;
+	case BLEND_Masked:
+		RHICmdList.SetBlendState(TStaticBlendState<>::GetRHI());
+		break;
+	case BLEND_Translucent:
+		RHICmdList.SetBlendState(TStaticBlendState<CW_RGB, BO_Add, BF_SourceAlpha, BF_InverseSourceAlpha, BO_Add, BF_Zero, BF_InverseSourceAlpha>::GetRHI());
+		break;
+	case BLEND_Additive:
+		// Add to the existing scene color
+		RHICmdList.SetBlendState(TStaticBlendState<CW_RGB, BO_Add, BF_One, BF_One, BO_Add, BF_Zero, BF_InverseSourceAlpha>::GetRHI());
+		break;
+	case BLEND_Modulate:
+		// Modulate with the existing scene color
+		RHICmdList.SetBlendState(TStaticBlendState<CW_RGB, BO_Add, BF_DestColor, BF_Zero>::GetRHI());
+		break;
+	};
+
+	SetShaderValue( RHICmdList, ShaderRHI, DisplayGamma, InDisplayGamma );
+	SetShaderValue( RHICmdList, ShaderRHI, ShaderParams, InShaderParams );
+
+	const bool bDeferredPass = false;
+	FMaterialShader::SetParameters<FPixelShaderRHIParamRef>(RHICmdList, ShaderRHI, MaterialRenderProxy, *Material, View, bDeferredPass, ESceneRenderTargetsMode::SetTextures);
+
+}
+
+
+bool FSlateMaterialShaderPS::Serialize(FArchive& Ar)
+{
+	bool bShaderHasOutdatedParameters = FMaterialShader::Serialize(Ar);
+	Ar << DisplayGamma;
+	Ar << ShaderParams;
+	return bShaderHasOutdatedParameters;
+}
+
+#define IMPLEMENT_SLATE_MATERIALSHADER_TYPE(ShaderType, bDrawDisabledEffect) \
+	typedef TSlateMaterialShaderPS<ESlateShader::ShaderType, bDrawDisabledEffect> TSlateMaterialShaderPS##ShaderType##bDrawDisabledEffect; \
+	IMPLEMENT_MATERIAL_SHADER_TYPE(template<>, TSlateMaterialShaderPS##ShaderType##bDrawDisabledEffect, TEXT("SlateElementPixelShader"), TEXT("Main"), SF_Pixel);
+
+
+IMPLEMENT_SLATE_MATERIALSHADER_TYPE(Default,true);
+IMPLEMENT_SLATE_MATERIALSHADER_TYPE(Default,false);
+IMPLEMENT_SLATE_MATERIALSHADER_TYPE(Border,true);
+IMPLEMENT_SLATE_MATERIALSHADER_TYPE(Border,false);
diff --git a/Engine/Source/Runtime/SlateRHIRenderer/Private/SlateMaterialShader.h b/Engine/Source/Runtime/SlateRHIRenderer/Private/SlateMaterialShader.h
new file mode 100644
--- /dev/null
+++ b/Engine/Source/Runtime/SlateRHIRenderer/Private/SlateMaterialShader.h
@@ -0,0 +1,72 @@
+// Copyright 1998-2014 Epic Games, Inc. All Rights Reserved.
+
+/*=============================================================================
+	PostProcessMaterial.h: Post processing Material
+=============================================================================*/
+
+#pragma once
+
+#include "Runtime/Renderer/Public/MaterialShader.h"
+
+class FSlateMaterialShaderPS : public FMaterialShader
+{
+public:
+	/**
+	  * Only compile shaders used with UI
+	  */
+	static bool ShouldCache(EShaderPlatform Platform, const FMaterial* Material);
+
+	/**
+	 * Modifies the compilation of this shader
+	 */
+	static void ModifyCompilationEnvironment(EShaderPlatform Platform, const FMaterial* Material, FShaderCompilerEnvironment& OutEnvironment);
+
+	FSlateMaterialShaderPS() {}
+	FSlateMaterialShaderPS(const FMaterialShaderType::CompiledShaderInitializerType& Initializer);
+
+	void SetParameters(FRHICommandList& RHICmdList, const FSceneView& View, const FMaterialRenderProxy* MaterialRenderProxy, const FMaterial* Material, float InDisplayGamma, const FVector4& InShaderParams );
+
+	virtual bool Serialize(FArchive& Ar) override;
+private:
+	FShaderParameter DisplayGamma;
+	FShaderParameter ShaderParams;
+};
+
+template<ESlateShader::Type ShaderType,bool bDrawDisabledEffect> 
+class TSlateMaterialShaderPS : public FSlateMaterialShaderPS
+{
+public:
+	DECLARE_SHADER_TYPE(TSlateMaterialShaderPS,Material);
+
+	TSlateMaterialShaderPS() {}
+	TSlateMaterialShaderPS(const ShaderMetaType::CompiledShaderInitializerType& Initializer)
+		: FSlateMaterialShaderPS( Initializer )
+	{
+	
+	}
+	
+	/**
+	  * Only compile shaders used with UI
+	  */
+	static bool ShouldCache(EShaderPlatform Platform, const FMaterial* Material)
+	{
+		return FSlateMaterialShaderPS::ShouldCache(Platform,Material);
+	}
+
+
+	/**
+	 * Modifies the compilation of this shader
+	 */
+	static void ModifyCompilationEnvironment(EShaderPlatform Platform, const FMaterial* Material, FShaderCompilerEnvironment& OutEnvironment)
+	{
+		FSlateMaterialShaderPS::ModifyCompilationEnvironment(Platform, Material,OutEnvironment);
+
+		OutEnvironment.SetDefine(TEXT("SHADER_TYPE"), (uint32)ShaderType);
+		OutEnvironment.SetDefine(TEXT("DRAW_DISABLED_EFFECT"), (uint32)(bDrawDisabledEffect ? 1 : 0));
+	}
+
+	virtual bool Serialize(FArchive& Ar) override
+	{
+		return FSlateMaterialShaderPS::Serialize( Ar );
+	}
+};
\ No newline at end of file
diff --git a/Engine/Source/Runtime/SlateRHIRenderer/Private/SlateRHIRenderingPolicy.cpp b/Engine/Source/Runtime/SlateRHIRenderer/Private/SlateRHIRenderingPolicy.cpp
--- a/Engine/Source/Runtime/SlateRHIRenderer/Private/SlateRHIRenderingPolicy.cpp
+++ b/Engine/Source/Runtime/SlateRHIRenderer/Private/SlateRHIRenderingPolicy.cpp
@@ -100,9 +100,6 @@ void FSlateElementIndexBuffer::ReleaseDynamicRHI()
 	BufferSize = 0;
 }
 
-FGlobalBoundShaderState FSlateRHIRenderingPolicy::NormalShaderStates[4][2 /* UseTextureAlpha */];
-FGlobalBoundShaderState FSlateRHIRenderingPolicy::DisabledShaderStates[4][2 /* UseTextureAlpha */];
-
 FSlateRHIRenderingPolicy::FSlateRHIRenderingPolicy( TSharedPtr<FSlateFontCache> InFontCache, TSharedRef<FSlateRHIResourceManager> InResourceManager )
 	: FSlateRenderingPolicy( GPixelCenterOffset )
 	, ResourceManager( InResourceManager )
@@ -216,6 +213,53 @@ static FSceneView& CreateSceneView( FSceneViewFamilyContext& ViewFamilyContext,
 	FSceneView* View = new FSceneView( ViewInitOptions );
 	ViewFamilyContext.Views.Add( View );
 
+	/** The view transform, starting from world-space points translated by -ViewOrigin. */
+	FMatrix EffectiveTranslatedViewMatrix = FTranslationMatrix(-View->ViewMatrices.PreViewTranslation) * View->ViewMatrices.ViewMatrix;
+
+	// Create the view's uniform buffer.
+	FViewUniformShaderParameters ViewUniformShaderParameters;
+	ViewUniformShaderParameters.TranslatedWorldToClip = View->ViewMatrices.TranslatedViewProjectionMatrix;
+	ViewUniformShaderParameters.WorldToClip = ViewProjectionMatrix;
+	ViewUniformShaderParameters.TranslatedWorldToView = EffectiveTranslatedViewMatrix;
+	ViewUniformShaderParameters.ViewToTranslatedWorld = View->InvViewMatrix * FTranslationMatrix(View->ViewMatrices.PreViewTranslation);
+	ViewUniformShaderParameters.ViewToClip = View->ViewMatrices.ProjMatrix;
+	ViewUniformShaderParameters.ClipToTranslatedWorld = View->ViewMatrices.InvTranslatedViewProjectionMatrix;
+	ViewUniformShaderParameters.ViewForward = EffectiveTranslatedViewMatrix.GetColumn(2);
+	ViewUniformShaderParameters.ViewUp = EffectiveTranslatedViewMatrix.GetColumn(1);
+	ViewUniformShaderParameters.ViewRight = EffectiveTranslatedViewMatrix.GetColumn(0);
+	ViewUniformShaderParameters.InvDeviceZToWorldZTransform = View->InvDeviceZToWorldZTransform;
+	ViewUniformShaderParameters.ScreenPositionScaleBias = FVector4(0,0,0,0);
+	ViewUniformShaderParameters.ScreenTexelBias = FVector4(ViewRect.Min.X, ViewRect.Min.Y, 0.0f, 0.0f);
+	ViewUniformShaderParameters.ViewSizeAndSceneTexelSize = FVector4(ViewRect.Width(), ViewRect.Height(), 1.0f/ViewRect.Width(), 1.0f/ViewRect.Height() );
+	ViewUniformShaderParameters.ViewOrigin = View->ViewMatrices.ViewOrigin;
+	ViewUniformShaderParameters.TranslatedViewOrigin = View->ViewMatrices.ViewOrigin + View->ViewMatrices.PreViewTranslation;
+	ViewUniformShaderParameters.DiffuseOverrideParameter = View->DiffuseOverrideParameter;
+	ViewUniformShaderParameters.SpecularOverrideParameter = View->SpecularOverrideParameter;
+	ViewUniformShaderParameters.NormalOverrideParameter = View->NormalOverrideParameter;
+	ViewUniformShaderParameters.RoughnessOverrideParameter = View->RoughnessOverrideParameter;
+	ViewUniformShaderParameters.PreViewTranslation = View->ViewMatrices.PreViewTranslation;
+	ViewUniformShaderParameters.CullingSign = View->bReverseCulling ? -1.0f : 1.0f;
+	ViewUniformShaderParameters.NearPlane = GNearClippingPlane;
+	ViewUniformShaderParameters.GameTime = View->Family->CurrentWorldTime;
+	ViewUniformShaderParameters.RealTime = View->Family->CurrentRealTime;
+	ViewUniformShaderParameters.Random = FMath::Rand();
+	ViewUniformShaderParameters.FrameNumber = View->FrameNumber;
+
+	ViewUniformShaderParameters.ScreenToWorld = FMatrix(
+		FPlane(1, 0, 0, 0),
+		FPlane(0, 1, 0, 0),
+		FPlane(0, 0, View->ProjectionMatrixUnadjustedForRHI.M[2][2], 1),
+		FPlane(0, 0, View->ProjectionMatrixUnadjustedForRHI.M[3][2], 0))
+		* View->InvViewProjectionMatrix;
+
+	ViewUniformShaderParameters.ScreenToTranslatedWorld = FMatrix(
+		FPlane(1, 0, 0, 0),
+		FPlane(0, 1, 0, 0),
+		FPlane(0, 0, View->ProjectionMatrixUnadjustedForRHI.M[2][2], 1),
+		FPlane(0, 0, View->ProjectionMatrixUnadjustedForRHI.M[3][2], 0))
+		* View->ViewMatrices.InvTranslatedViewProjectionMatrix;
+
+	View->UniformBuffer = TUniformBufferRef<FViewUniformShaderParameters>::CreateUniformBufferImmediate(ViewUniformShaderParameters, UniformBuffer_SingleFrame);
 	return *View;
 }
 
@@ -240,7 +284,6 @@ void FSlateRHIRenderingPolicy::DrawElements( const FIntPoint& InViewportSize, FS
 
 	const float DisplayGamma = GEngine ? GEngine->GetDisplayGamma() : 2.2f;
 
-#if ALLOW_SLATE_MATERIALS 
 	FSceneViewFamilyContext SceneViewContext
 	(
 		FSceneViewFamily::ConstructionValues
@@ -254,7 +297,6 @@ void FSlateRHIRenderingPolicy::DrawElements( const FIntPoint& InViewportSize, FS
 	);
 
 	FSceneView* SceneView = NULL;
-#endif
 
 	TShaderMapRef<FSlateElementVS> VertexShader(GetGlobalShaderMap());
 
@@ -275,131 +317,152 @@ void FSlateRHIRenderingPolicy::DrawElements( const FIntPoint& InViewportSize, FS
 		const ESlateDrawEffect::Type DrawEffects = RenderBatch.DrawEffects;
 		const ESlateShader::Type ShaderType = RenderBatch.ShaderType;
 		const FShaderParams& ShaderParams = RenderBatch.ShaderParams;
-		const bool bDrawDisabled = (DrawEffects & ESlateDrawEffect::DisabledEffect) != 0;
 
-		if( !RenderBatch.CustomDrawer.IsValid() && (!ShaderResource || ShaderResource->GetType() == ESlateShaderResource::Texture ) )
+		if( !RenderBatch.CustomDrawer.IsValid() )
 		{
-			FSlateElementPS* PixelShader = GetPixelShader( ShaderType, DrawEffects );
-	
-			const bool bUseTextureAlpha = (DrawEffects & ESlateDrawEffect::IgnoreTextureAlpha) == 0;
-
-			if( bDrawDisabled )
+			if( !ShaderResource || ShaderResource->GetType() == ESlateShaderResource::Texture ) 
 			{
-				SetGlobalBoundShaderState(RHICmdList, DisabledShaderStates[ShaderType][bUseTextureAlpha], GSlateVertexDeclaration.VertexDeclarationRHI, *VertexShader, PixelShader);
-			}
-			else
-			{
-				SetGlobalBoundShaderState(RHICmdList, NormalShaderStates[ShaderType][bUseTextureAlpha], GSlateVertexDeclaration.VertexDeclarationRHI, *VertexShader, PixelShader);
-			}
+				FSlateElementPS* PixelShader = GetTexturePixelShader(ShaderType, DrawEffects);
+				FBoundShaderStateRHIRef BoundShaderState;
+
+				BoundShaderState = RHICreateBoundShaderState(
+					GSlateVertexDeclaration.VertexDeclarationRHI,
+					VertexShader->GetVertexShader(),
+					nullptr,
+					nullptr,
+					PixelShader->GetPixelShader(),
+					FGeometryShaderRHIRef());
+
+				RHISetBoundShaderState(BoundShaderState);
 
-			VertexShader->SetViewProjection(RHICmdList, ViewProjectionMatrix );
-			VertexShader->SetShaderParameters(RHICmdList, ShaderParams.VertexParams );
+				VertexShader->SetViewProjection(RHICmdList, ViewProjectionMatrix);
+				VertexShader->SetShaderParameters(RHICmdList, ShaderParams.VertexParams);
 
 #if !DEBUG_OVERDRAW
-			RHICmdList.SetBlendState(
-				(RenderBatch.DrawFlags & ESlateBatchDrawFlag::NoBlending)
-				? TStaticBlendState<>::GetRHI()
-				: TStaticBlendState<CW_RGBA,BO_Add,BF_SourceAlpha,BF_InverseSourceAlpha,BO_Add,BF_Zero,BF_One>::GetRHI()
-				);
+				RHICmdList.SetBlendState(
+					(RenderBatch.DrawFlags & ESlateBatchDrawFlag::NoBlending)
+					? TStaticBlendState<>::GetRHI()
+					: TStaticBlendState<CW_RGBA, BO_Add, BF_SourceAlpha, BF_InverseSourceAlpha, BO_Add, BF_Zero, BF_One>::GetRHI()
+					);
 #else
-			RHICmdList.SetBlendState(TStaticBlendState<CW_RGB, BO_Add, BF_One, BF_One, BO_Add, BF_Zero, BF_InverseSourceAlpha>::GetRHI());
+				RHICmdList.SetBlendState(TStaticBlendState<CW_RGB, BO_Add, BF_One, BF_One, BO_Add, BF_Zero, BF_InverseSourceAlpha>::GetRHI());
 #endif
 
-			// Disable stencil testing by default
-			RHICmdList.SetDepthStencilState(DSOff);
+				// Disable stencil testing by default
+				RHICmdList.SetDepthStencilState(DSOff);
 
-			if( DrawFlags & ESlateBatchDrawFlag::Wireframe )
-			{
-				RHICmdList.SetRasterizerState(TStaticRasterizerState<FM_Wireframe, CM_None, true>::GetRHI());
-			}
-			else
-			{
-				RHICmdList.SetRasterizerState(TStaticRasterizerState<FM_Solid, CM_None, true>::GetRHI());
-			}
+				if (DrawFlags & ESlateBatchDrawFlag::Wireframe)
+				{
+					RHICmdList.SetRasterizerState(TStaticRasterizerState<FM_Wireframe, CM_None, true>::GetRHI());
+				}
+				else
+				{
+					RHICmdList.SetRasterizerState(TStaticRasterizerState<FM_Solid, CM_None, true>::GetRHI());
+				}
 
 
-			if( ShaderType == ESlateShader::LineSegment )
-			{
-// Generally looks better without stencil testing although there is minor overdraw on splines
+				if (ShaderType == ESlateShader::LineSegment)
+				{
+					// Generally looks better without stencil testing although there is minor overdraw on splines
 #if 0
-				// Enable stencil testing for anti-aliased line segments to reduce artifacts from overlapping segments
-				/FDepthStencilStateRHIRef DSLineSegment = 
-					TStaticDepthStencilState< false
-					, CF_Always
-					, true
-					, CF_Greater
-					, SO_Keep
-					, SO_SaturatedIncrement
-					, SO_SaturatedIncrement
-					, true
-					, CF_Greater
-					, SO_Keep
-					, SO_SaturatedIncrement
-					, SO_SaturatedIncrement>::GetRHI();
-
-				RHICmdList.SetDepthStencilState( DSLineSegment, 0x01 );
+					// Enable stencil testing for anti-aliased line segments to reduce artifacts from overlapping segments
+					 FDepthStencilStateRHIRef DSLineSegment =
+						TStaticDepthStencilState < false
+						, CF_Always
+						, true
+						, CF_Greater
+						, SO_Keep
+						, SO_SaturatedIncrement
+						, SO_SaturatedIncrement
+						, true
+						, CF_Greater
+						, SO_Keep
+						, SO_SaturatedIncrement
+						, SO_SaturatedIncrement>::GetRHI();
+
+					RHICmdList.SetDepthStencilState( DSLineSegment, 0x01 );
 #endif
 
-				// The lookup table used for splines should clamp when sampling otherwise the results are wrong
-				PixelShader->SetTexture(RHICmdList, ((FSlateTexture2DRHIRef*)ShaderResource)->GetTypedResource(), BilinearClamp );
-			}
-			else if( ShaderResource && ShaderResource->GetType() == ESlateShaderResource::Texture && IsValidRef( ((FSlateTexture2DRHIRef*)ShaderResource)->GetTypedResource() ) )
-			{
-				FSamplerStateRHIRef SamplerState; 
-
-				if( DrawFlags == (ESlateBatchDrawFlag::TileU | ESlateBatchDrawFlag::TileV) )
-				{
-					SamplerState = TStaticSamplerState<SF_Bilinear,AM_Wrap,AM_Wrap,AM_Wrap>::GetRHI();
-				}
-				else if( DrawFlags & ESlateBatchDrawFlag::TileU )
-				{
-					SamplerState =  TStaticSamplerState<SF_Bilinear,AM_Wrap,AM_Clamp,AM_Wrap>::GetRHI();
+					// The lookup table used for splines should clamp when sampling otherwise the results are wrong
+					PixelShader->SetTexture(RHICmdList, ((FSlateTexture2DRHIRef*)ShaderResource)->GetTypedResource(), BilinearClamp );
 				}
-				else if( DrawFlags & ESlateBatchDrawFlag::TileV )
+				else if( ShaderResource && ShaderResource->GetType() == ESlateShaderResource::Texture && IsValidRef( ((FSlateTexture2DRHIRef*)ShaderResource)->GetTypedResource() ) )
 				{
-					SamplerState = TStaticSamplerState<SF_Bilinear,AM_Clamp,AM_Wrap,AM_Wrap>::GetRHI();
+					FSamplerStateRHIRef SamplerState; 
+
+					if( DrawFlags == (ESlateBatchDrawFlag::TileU | ESlateBatchDrawFlag::TileV) )
+					{
+						SamplerState = TStaticSamplerState<SF_Bilinear, AM_Wrap, AM_Wrap, AM_Wrap>::GetRHI();
+					}
+					else if (DrawFlags & ESlateBatchDrawFlag::TileU)
+					{
+						SamplerState = TStaticSamplerState<SF_Bilinear, AM_Wrap, AM_Clamp, AM_Wrap>::GetRHI();
+					}
+					else if (DrawFlags & ESlateBatchDrawFlag::TileV)
+					{
+						SamplerState = TStaticSamplerState<SF_Bilinear, AM_Clamp, AM_Wrap, AM_Wrap>::GetRHI();
+					}
+					else
+					{
+						SamplerState = BilinearClamp;
+					}
+					PixelShader->SetTexture(RHICmdList, ((FSlateTexture2DRHIRef*)ShaderResource)->GetTypedResource(), SamplerState);
 				}
 				else
 				{
-					SamplerState = BilinearClamp;
+					PixelShader->SetTexture(RHICmdList, GWhiteTexture->TextureRHI, BilinearClamp);
 				}
-				PixelShader->SetTexture(RHICmdList, ((FSlateTexture2DRHIRef*)ShaderResource)->GetTypedResource(), SamplerState );
+
+				PixelShader->SetShaderParams(RHICmdList, ShaderParams.PixelParams);
+				PixelShader->SetDisplayGamma(RHICmdList, (DrawFlags & ESlateBatchDrawFlag::NoGamma) ? 1.0f : 1.0f/DisplayGamma);
+
+
+				check(RenderBatch.NumIndices > 0);
+
+				uint32 PrimitiveCount = RenderBatch.DrawPrimitiveType == ESlateDrawPrimitive::LineList ? RenderBatch.NumIndices / 2 : RenderBatch.NumIndices / 3;
+
+				RHICmdList.DrawIndexedPrimitive(IndexBuffer.IndexBufferRHI, GetRHIPrimitiveType(RenderBatch.DrawPrimitiveType), RenderBatch.VertexOffset, 0, RenderBatch.NumVertices, RenderBatch.IndexOffset, PrimitiveCount, 1);
+
 			}
-			else
+			else if (ShaderResource && ShaderResource->GetType() == ESlateShaderResource::Material)
 			{
-				PixelShader->SetTexture(RHICmdList, GWhiteTexture->TextureRHI, BilinearClamp );
-			}
+				if (!SceneView)
+				{
+					SceneView = &CreateSceneView(SceneViewContext, BackBuffer, ViewProjectionMatrix);
+				}
 
-			PixelShader->SetShaderParams(RHICmdList, ShaderParams.PixelParams );
-			PixelShader->SetDisplayGamma(RHICmdList, ( DrawFlags & ESlateBatchDrawFlag::NoGamma ) ? 1.0f : DisplayGamma );
+				FMaterialRenderProxy* MaterialRenderProxy = ((FSlateMaterial*)ShaderResource)->GetRenderProxy();
 
-			check( RenderBatch.NumIndices > 0 );
+				const FMaterial* Material = MaterialRenderProxy->GetMaterial(SceneView->GetFeatureLevel());
 
-			uint32 PrimitiveCount = RenderBatch.DrawPrimitiveType == ESlateDrawPrimitive::LineList ? RenderBatch.NumIndices / 2 : RenderBatch.NumIndices / 3; 
+				FSlateMaterialShaderPS* PixelShader = GetMaterialPixelShader( Material, ShaderType, DrawEffects );
 
-			RHICmdList.DrawIndexedPrimitive(IndexBuffer.IndexBufferRHI, GetRHIPrimitiveType(RenderBatch.DrawPrimitiveType), RenderBatch.VertexOffset, 0, RenderBatch.NumVertices, RenderBatch.IndexOffset, PrimitiveCount, 1);
-		}
-		else if( ShaderResource && ShaderResource->GetType() == ESlateShaderResource::Material )
-		{
-#if ALLOW_SLATE_MATERIALS
-			FMaterialRenderProxy* MaterialRenderProxy = ((FSlateMaterial*)ShaderResource)->GetRenderProxy();
-		
-			if( !SceneView )
-			{
-				SceneView = &CreateSceneView( SceneViewContext, BackBuffer, ViewProjectionMatrix );
-			}
+				if( PixelShader )
+				{
+					FBoundShaderStateRHIRef BoundShaderState;
 
-		
-			// Bind Shader here
+					BoundShaderState = RHICreateBoundShaderState(
+						GSlateVertexDeclaration.VertexDeclarationRHI,
+						VertexShader->GetVertexShader(),
+						nullptr,
+						nullptr,
+						PixelShader->GetPixelShader(),
+						FGeometryShaderRHIRef());
 
-			VertexShader->SetShaderParameters(RHICmdList, ShaderParams.VertexParams);
+					RHISetBoundShaderState(BoundShaderState);
 
-			check(RenderBatch.NumIndices > 0);
+					VertexShader->SetShaderParameters(RHICmdList, ShaderParams.VertexParams);
+					PixelShader->SetParameters(RHICmdList, *SceneView, MaterialRenderProxy, Material, 1.0f / DisplayGamma, ShaderParams.PixelParams);
 
-			uint32 PrimitiveCount = RenderBatch.DrawPrimitiveType == ESlateDrawPrimitive::LineList ? RenderBatch.NumIndices / 2 : RenderBatch.NumIndices / 3;
 
-			RHICmdList.DrawIndexedPrimitive(IndexBuffer.IndexBufferRHI, GetRHIPrimitiveType(RenderBatch.DrawPrimitiveType), RenderBatch.VertexOffset, 0, RenderBatch.NumVertices, RenderBatch.IndexOffset, PrimitiveCount, 1);
-#endif
+					check(RenderBatch.NumIndices > 0);
+
+					uint32 PrimitiveCount = RenderBatch.DrawPrimitiveType == ESlateDrawPrimitive::LineList ? RenderBatch.NumIndices / 2 : RenderBatch.NumIndices / 3;
+
+					RHICmdList.DrawIndexedPrimitive(IndexBuffer.IndexBufferRHI, GetRHIPrimitiveType(RenderBatch.DrawPrimitiveType), RenderBatch.VertexOffset, 0, RenderBatch.NumVertices, RenderBatch.IndexOffset, PrimitiveCount, 1);
+				}
+			}
+
 		}
 		else if (RenderBatch.CustomDrawer.IsValid())
 		{
@@ -413,7 +476,7 @@ void FSlateRHIRenderingPolicy::DrawElements( const FIntPoint& InViewportSize, FS
 }
 
 
-FSlateElementPS* FSlateRHIRenderingPolicy::GetPixelShader( ESlateShader::Type ShaderType, ESlateDrawEffect::Type DrawEffects )
+FSlateElementPS* FSlateRHIRenderingPolicy::GetTexturePixelShader( ESlateShader::Type ShaderType, ESlateDrawEffect::Type DrawEffects )
 {
 	FSlateElementPS* PixelShader = NULL;
 
@@ -492,3 +555,44 @@ FSlateElementPS* FSlateRHIRenderingPolicy::GetPixelShader( ESlateShader::Type Sh
 
 	return PixelShader;
 }
+
+FSlateMaterialShaderPS* FSlateRHIRenderingPolicy::GetMaterialPixelShader( const FMaterial* Material, ESlateShader::Type ShaderType, ESlateDrawEffect::Type DrawEffects )
+{
+	const FMaterialShaderMap* MaterialShaderMap = Material->GetRenderingThreadShaderMap();
+
+	const bool bDrawDisabled = (DrawEffects & ESlateDrawEffect::DisabledEffect) != 0;
+	const bool bUseTextureAlpha = (DrawEffects & ESlateDrawEffect::IgnoreTextureAlpha) == 0;
+
+	FShader* FoundShader = nullptr;
+	switch (ShaderType)
+	{
+	case ESlateShader::Default:
+		if (bDrawDisabled)
+		{
+			FoundShader = MaterialShaderMap->GetShader(&TSlateMaterialShaderPS<0, true>::StaticType);
+		}
+		else
+		{
+			FoundShader = MaterialShaderMap->GetShader(&TSlateMaterialShaderPS<0, false>::StaticType);
+		}
+		break;
+	case ESlateShader::Border:
+		if (bDrawDisabled)
+		{
+			FoundShader = MaterialShaderMap->GetShader(&TSlateMaterialShaderPS<1, true>::StaticType);
+		}
+		else
+		{
+			FoundShader = MaterialShaderMap->GetShader(&TSlateMaterialShaderPS<1, false>::StaticType);
+		}
+		break;
+	default:
+		checkf(false, TEXT("Unsupported Slate shader type for use with materials"));
+		break;
+	}
+
+	return FoundShader ? (FSlateMaterialShaderPS*)FoundShader->GetShaderChecked() : nullptr;
+}
+
+
+
diff --git a/Engine/Source/Runtime/SlateRHIRenderer/Private/SlateRHIRenderingPolicy.h b/Engine/Source/Runtime/SlateRHIRenderer/Private/SlateRHIRenderingPolicy.h
--- a/Engine/Source/Runtime/SlateRHIRenderer/Private/SlateRHIRenderingPolicy.h
+++ b/Engine/Source/Runtime/SlateRHIRenderer/Private/SlateRHIRenderingPolicy.h
@@ -195,13 +195,10 @@ private:
 	 * @param DrawEffects	Draw effects being used
 	 * @return The pixel shader for use with the shader type and draw effects
 	 */
-	class FSlateElementPS* GetPixelShader( ESlateShader::Type ShaderType, ESlateDrawEffect::Type DrawEffects );
-private:
-	/** Global shader states */
-	static FGlobalBoundShaderState NormalShaderStates[4][2 /* UseTextureAlpha */];
-	static FGlobalBoundShaderState DisabledShaderStates[4][2 /* UseTextureAlpha */];
-	static FGlobalBoundShaderState MaterialShaderState;
+	class FSlateElementPS* GetTexturePixelShader( ESlateShader::Type ShaderType, ESlateDrawEffect::Type DrawEffects );
+	class FSlateMaterialShaderPS* GetMaterialPixelShader( const class FMaterial* Material, ESlateShader::Type ShaderType, ESlateDrawEffect::Type DrawEffects );
 
+private:
 	/** Buffers used for rendering */
 	TSlateElementVertexBuffer<FSlateVertex> VertexBuffers[2];
 	FSlateElementIndexBuffer IndexBuffers[2];
diff --git a/Engine/Source/Runtime/SlateRHIRenderer/Private/SlateRHIResourceManager.cpp b/Engine/Source/Runtime/SlateRHIRenderer/Private/SlateRHIResourceManager.cpp
--- a/Engine/Source/Runtime/SlateRHIRenderer/Private/SlateRHIResourceManager.cpp
+++ b/Engine/Source/Runtime/SlateRHIRenderer/Private/SlateRHIResourceManager.cpp
@@ -168,16 +168,6 @@ void FDynamicResourceMap::AddReferencedObjects(FReferenceCollector& Collector)
 			Collector.AddReferencedObject(Resource->TextureObject);
 		}
 	}
-
-
-	for (TMap<UMaterialInterface*, TSharedPtr<FSlateMaterialResource> >::TIterator It(MaterialResourceMap); It; ++It)
-	{
-		TSharedPtr<FSlateMaterialResource>& Resource = It.Value();
-		if (Resource->MaterialObject)
-		{
-			Collector.AddReferencedObject(Resource->MaterialObject);
-		}
-	}
 }
 
 
@@ -632,6 +622,7 @@ FSlateShaderResourceProxy* FSlateRHIResourceManager::GetMaterialResource(const F
 	{
 		// Keep the resource up to date
 		MaterialResource->SlateMaterial->SetResource(Material->GetRenderProxy(false, false));
+		MaterialResource->Proxy->ActualSize = InBrush.ImageSize.IntPoint();
 	}
 
 	return MaterialResource->Proxy;
diff --git a/Engine/Source/Runtime/SlateRHIRenderer/Private/SlateRHIResourceManager.h b/Engine/Source/Runtime/SlateRHIRenderer/Private/SlateRHIResourceManager.h
--- a/Engine/Source/Runtime/SlateRHIRenderer/Private/SlateRHIResourceManager.h
+++ b/Engine/Source/Runtime/SlateRHIRenderer/Private/SlateRHIResourceManager.h
@@ -61,7 +61,7 @@ private:
 	TMap<UObject*, TSharedPtr<FDynamicTextureResource> > DynamicResourceObjectMap;
 
 	/** Map of all material resources */
-	TMap<UMaterialInterface*, TSharedPtr<FSlateMaterialResource> > MaterialResourceMap;
+	TMap<TWeakObjectPtr<UMaterialInterface>, TSharedPtr<FSlateMaterialResource> > MaterialResourceMap;
 };
 
 
diff --git a/Engine/Source/Runtime/SlateRHIRenderer/Private/SlateShaders.h b/Engine/Source/Runtime/SlateRHIRenderer/Private/SlateShaders.h
--- a/Engine/Source/Runtime/SlateRHIRenderer/Private/SlateShaders.h
+++ b/Engine/Source/Runtime/SlateRHIRenderer/Private/SlateShaders.h
@@ -88,7 +88,7 @@ public:
 		TextureParameter.Bind( Initializer.ParameterMap, TEXT("ElementTexture"));
 		TextureParameterSampler.Bind( Initializer.ParameterMap, TEXT("ElementTextureSampler"));
 		ShaderParams.Bind( Initializer.ParameterMap, TEXT("ShaderParams"));
-		DisplayGamma.Bind( Initializer.ParameterMap,TEXT("DisplayGamma"));
+		DisplayGamma.Bind( Initializer.ParameterMap,TEXT("InvDisplayGamma"));
 	}
 
 
diff --git a/Engine/Source/Runtime/SlateRHIRenderer/SlateRHIRenderer.Build.cs b/Engine/Source/Runtime/SlateRHIRenderer/SlateRHIRenderer.Build.cs
--- a/Engine/Source/Runtime/SlateRHIRenderer/SlateRHIRenderer.Build.cs
+++ b/Engine/Source/Runtime/SlateRHIRenderer/SlateRHIRenderer.Build.cs
@@ -18,7 +18,8 @@ public class SlateRHIRenderer : ModuleRules
                 "Engine",
                 "RHI",
                 "RenderCore",
-                "ShaderCore",
+				"ShaderCore",
+				"Renderer",
                 "ImageWrapper"
 			}
 		);
