UE4 - rework RHI user-facing APIs to be safer and incorporate the desired sematics

diff --git a/Engine/Source/Runtime/SlateRHIRenderer/Private/SlateRHIRenderer.cpp b/Engine/Source/Runtime/SlateRHIRenderer/Private/SlateRHIRenderer.cpp
--- a/Engine/Source/Runtime/SlateRHIRenderer/Private/SlateRHIRenderer.cpp
+++ b/Engine/Source/Runtime/SlateRHIRenderer/Private/SlateRHIRenderer.cpp
@@ -333,11 +333,11 @@ void FSlateRHIRenderer::ConditionalResizeViewport( FViewportInfo* ViewInfo, uint
 
 		if( IsValidRef( ViewInfo->ViewportRHI ) )
 		{
-			RHIResizeViewport( ViewInfo->ViewportRHI, NewWidth, NewHeight, bFullscreen );
+			RHIResizeViewport(ViewInfo->ViewportRHI, NewWidth, NewHeight, bFullscreen);
 		}
 		else
 		{
-			ViewInfo->ViewportRHI = RHICreateViewport( ViewInfo->OSWindow, NewWidth, NewHeight, bFullscreen );
+			ViewInfo->ViewportRHI = RHICreateViewport(ViewInfo->OSWindow, NewWidth, NewHeight, bFullscreen);
 		}
 
 		// Safe to call here as the rendering thread has been suspended: game thread == render thread!
@@ -394,7 +394,7 @@ void FSlateRHIRenderer::OnWindowDestroyed( const TSharedRef<SWindow>& InWindow )
 }
 
 /** Draws windows from a FSlateDrawBuffer on the render thread */
-void FSlateRHIRenderer::DrawWindow_RenderThread( const FViewportInfo& ViewportInfo, const FSlateWindowElementList& WindowElementList, bool bLockToVsync )
+void FSlateRHIRenderer::DrawWindow_RenderThread(FRHICommandListImmediate& RHICmdList, const FViewportInfo& ViewportInfo, const FSlateWindowElementList& WindowElementList, bool bLockToVsync)
 {
 	SCOPED_DRAW_EVENT(SlateUI, DEC_SCENE_ITEMS);
 
@@ -409,23 +409,23 @@ void FSlateRHIRenderer::DrawWindow_RenderThread( const FViewportInfo& ViewportIn
 		// should have been created by the game thread
 		check( IsValidRef(ViewportInfo.ViewportRHI) );
 
-		RHIBeginDrawingViewport( ViewportInfo.ViewportRHI, FTextureRHIRef() );
-		RHISetViewport( 0,0,0,ViewportInfo.Width, ViewportInfo.Height, 0.0f );
+		RHICmdList.BeginDrawingViewport(ViewportInfo.ViewportRHI, FTextureRHIRef());
+		RHICmdList.SetViewport( 0,0,0,ViewportInfo.Width, ViewportInfo.Height, 0.0f );
 
-		FTexture2DRHIRef BackBuffer = RHIGetViewportBackBuffer( ViewportInfo.ViewportRHI );
+		FTexture2DRHIRef BackBuffer = RHICmdList.GetViewportBackBuffer(ViewportInfo.ViewportRHI);
 
 		if( ViewportInfo.bRequiresStencilTest )
 		{
 			check(IsValidRef( ViewportInfo.DepthStencil ));
 
 			// Reset the backbuffer as our color render target and also set a depth stencil buffer
-			RHISetRenderTarget( BackBuffer, ViewportInfo.DepthStencil );
+			SetRenderTarget(RHICmdList,  BackBuffer, ViewportInfo.DepthStencil );
 			// Clear the stencil buffer
-			RHIClear( false, FLinearColor::White, false, 0.0f, true, 0x00, FIntRect());
+			RHICmdList.Clear( false, FLinearColor::White, false, 0.0f, true, 0x00, FIntRect());
 		}
 		else
 		{
-			RHISetRenderTarget(BackBuffer, FTextureRHIRef());
+			SetRenderTarget(RHICmdList, BackBuffer, FTextureRHIRef());
 		}
 
 #if DEBUG_OVERDRAW
@@ -437,6 +437,7 @@ void FSlateRHIRenderer::DrawWindow_RenderThread( const FViewportInfo& ViewportIn
 
 			RenderingPolicy->DrawElements
 			(
+				RHICmdList,
 				FIntPoint(ViewportInfo.Width, ViewportInfo.Height),
 				BackBufferTarget,
 				ViewMatrix*ViewportInfo.ProjectionMatrix,
@@ -450,7 +451,7 @@ void FSlateRHIRenderer::DrawWindow_RenderThread( const FViewportInfo& ViewportIn
 	uint32 StartTime		= FPlatformTime::Cycles();
 		
 	// Note - We do not include present time in the slate render thread stat
-	RHIEndDrawingViewport( ViewportInfo.ViewportRHI, true, bLockToVsync );
+	RHICmdList.EndDrawingViewport(ViewportInfo.ViewportRHI, true, bLockToVsync);
 
 	uint32 EndTime		= FPlatformTime::Cycles();
 
@@ -619,7 +620,7 @@ void FSlateRHIRenderer::DrawWindows_Private( FSlateDrawBuffer& WindowDrawBuffer
 					ENQUEUE_UNIQUE_RENDER_COMMAND_ONEPARAMETER( SlateDrawWindowsCommand, 
 						FSlateDrawWindowCommandParams, Params, Params,
 					{
-						Params.Renderer->DrawWindow_RenderThread( *Params.ViewportInfo, *Params.WindowElementList, Params.bLockToVsync );
+						Params.Renderer->DrawWindow_RenderThread(RHICmdList, *Params.ViewportInfo, *Params.WindowElementList, Params.bLockToVsync);
 					});
 
 					SlateWindowRendered.Broadcast( *Params.SlateWindow, &ViewInfo->ViewportRHI );
@@ -631,8 +632,8 @@ void FSlateRHIRenderer::DrawWindows_Private( FSlateDrawBuffer& WindowDrawBuffer
 							FIntRect, ScreenshotRect, ScreenshotRect,
 							TArray<FColor>*, OutScreenshotData, OutScreenshotData,
 						{
-							FTexture2DRHIRef BackBuffer = RHIGetViewportBackBuffer(Params.ViewportInfo->ViewportRHI);
-							RHIReadSurfaceData(BackBuffer, ScreenshotRect, *OutScreenshotData, FReadSurfaceDataFlags());
+							FTexture2DRHIRef BackBuffer = RHICmdList.GetViewportBackBuffer(Params.ViewportInfo->ViewportRHI);
+							RHICmdList.ReadSurfaceData(BackBuffer, ScreenshotRect, *OutScreenshotData, FReadSurfaceDataFlags());
 						});
 
 						FlushRenderingCommands();
@@ -774,15 +775,19 @@ void FSlateRHIRenderer::CopyWindowsToVirtualScreenBuffer(const TArray<FString>&
 		SetupWindowState,
 		FSetupWindowStateContext,Context,SetupWindowStateContext,
 	{
-		RHISetRenderTarget(Context.CrashReportResource->GetBuffer(), FTextureRHIRef());
+		SetRenderTarget(RHICmdList, Context.CrashReportResource->GetBuffer(), FTextureRHIRef());		
+		RHICmdList.SetViewport(0, 0, 0.0f, Context.IntermediateBufferSize.Width(), Context.IntermediateBufferSize.Height(), 1.0f);
+		RHICmdList.SetBlendState(TStaticBlendState<>::GetRHI());
+		RHICmdList.SetRasterizerState(TStaticRasterizerState<>::GetRHI());
+		RHICmdList.SetDepthStencilState(TStaticDepthStencilState<false, CF_Always>::GetRHI());
 		
-		RHISetViewport(0, 0, 0.0f, Context.IntermediateBufferSize.Width(), Context.IntermediateBufferSize.Height(), 1.0f);
-		RHISetBlendState(TStaticBlendState<>::GetRHI());
-		RHISetRasterizerState(TStaticRasterizerState<>::GetRHI());
-		RHISetDepthStencilState(TStaticDepthStencilState<false,CF_Always>::GetRHI());
+		RHICmdList.SetViewport(0, 0, 0.0f, Context.IntermediateBufferSize.Width(), Context.IntermediateBufferSize.Height(), 1.0f);
+		RHICmdList.SetBlendState(TStaticBlendState<>::GetRHI());
+		RHICmdList.SetRasterizerState(TStaticRasterizerState<>::GetRHI());
+		RHICmdList.SetDepthStencilState(TStaticDepthStencilState<false,CF_Always>::GetRHI());
 
 		// @todo livestream: Ideally this "desktop background color" should be configurable in the editor's preferences
-		RHIClear(true, FLinearColor( 0.02f, 0.02f, 0.2f ), false, 0.f, false, 0x00, FIntRect());
+		RHICmdList.Clear(true, FLinearColor(0.02f, 0.02f, 0.2f), false, 0.f, false, 0x00, FIntRect());
 	});
 
 	// draw windows to buffer
@@ -831,20 +836,18 @@ void FSlateRHIRenderer::CopyWindowsToVirtualScreenBuffer(const TArray<FString>&
 				TShaderMapRef<FScreenVS> VertexShader(GetGlobalShaderMap());
 				TShaderMapRef<FScreenPS> PixelShader(GetGlobalShaderMap());
 
-				FRHICommandList& RHICmdList = FRHICommandList::GetNullRef();
-
 				static FGlobalBoundShaderState BoundShaderState;
 				SetGlobalBoundShaderState(RHICmdList, BoundShaderState, Context.RendererModule->GetFilterVertexDeclaration().VertexDeclarationRHI, *VertexShader, *PixelShader);
 
 				if( Context.WindowRect.Width() != Context.InViewportInfo->Width || Context.WindowRect.Height() != Context.InViewportInfo->Height )
 				{
 					// We're scaling down the window, so use bilinear filtering
-					PixelShader->SetParameters(RHICmdList, TStaticSamplerState<SF_Bilinear>::GetRHI(), RHIGetViewportBackBuffer(Context.InViewportInfo->ViewportRHI));
+					PixelShader->SetParameters(RHICmdList, TStaticSamplerState<SF_Bilinear>::GetRHI(), RHICmdList.GetViewportBackBuffer(Context.InViewportInfo->ViewportRHI));
 				}
 				else
 				{
 					// Drawing 1:1, so no filtering needed
-					PixelShader->SetParameters(RHICmdList, TStaticSamplerState<SF_Point>::GetRHI(), RHIGetViewportBackBuffer(Context.InViewportInfo->ViewportRHI));
+					PixelShader->SetParameters(RHICmdList, TStaticSamplerState<SF_Point>::GetRHI(), RHICmdList.GetViewportBackBuffer(Context.InViewportInfo->ViewportRHI));
 				}
 
 				Context.RendererModule->DrawRectangle(
@@ -872,14 +875,14 @@ void FSlateRHIRenderer::CopyWindowsToVirtualScreenBuffer(const TArray<FString>&
 	// @todo livestream: The cursor is probably still hidden when dragging with the mouse captured (grabby hand)
 	if( FSlateApplication::Get().GetMouseCaptureWindow() == nullptr )
 	{
-		FSlateDrawElement::MakeBox(
-			*WindowElementList,
-			0,
+	FSlateDrawElement::MakeBox(
+		*WindowElementList,
+		0,
 			FPaintGeometry(ScaledCursorLocation, FVector2D(32, 32) * XScaling, XScaling),
-			FCoreStyle::Get().GetBrush("CrashTracker.Cursor"),
+		FCoreStyle::Get().GetBrush("CrashTracker.Cursor"),
 			FSlateRect(0, 0, VirtualScreenSize.X, VirtualScreenSize.Y));
 	}
-
+	
 	for (int32 i = 0; i < KeypressBuffer.Num(); ++i)
 	{
 		FSlateDrawElement::MakeText(
@@ -917,7 +920,7 @@ void FSlateRHIRenderer::CopyWindowsToVirtualScreenBuffer(const TArray<FString>&
 		WriteMouseCursorAndKeyPresses,
 		FWriteMouseCursorAndKeyPressesContext, Context, WriteMouseCursorAndKeyPressesContext,
 	{
-		RHISetBlendState(TStaticBlendState<CW_RGBA,BO_Add,BF_SourceAlpha,BF_InverseSourceAlpha,BO_Add,BF_Zero,BF_One>::GetRHI());
+		RHICmdList.SetBlendState(TStaticBlendState<CW_RGBA, BO_Add, BF_SourceAlpha, BF_InverseSourceAlpha, BO_Add, BF_Zero, BF_One>::GetRHI());
 		
 		Context.RenderPolicy->UpdateBuffers(*Context.SlateElementList);
 		if( Context.SlateElementList->GetRenderBatches().Num() > 0 )
@@ -925,7 +928,7 @@ void FSlateRHIRenderer::CopyWindowsToVirtualScreenBuffer(const TArray<FString>&
 			FTexture2DRHIRef UnusedTargetTexture;
 			FSlateBackBuffer UnusedTarget( UnusedTargetTexture, FIntPoint::ZeroValue );
 
-			Context.RenderPolicy->DrawElements(Context.ViewportSize, UnusedTarget, CreateProjectionMatrix(Context.ViewportSize.X, Context.ViewportSize.Y), Context.SlateElementList->GetRenderBatches());
+			Context.RenderPolicy->DrawElements(RHICmdList, Context.ViewportSize, UnusedTarget, CreateProjectionMatrix(Context.ViewportSize.X, Context.ViewportSize.Y), Context.SlateElementList->GetRenderBatches());
 		}
 	});
 
@@ -944,7 +947,7 @@ void FSlateRHIRenderer::CopyWindowsToVirtualScreenBuffer(const TArray<FString>&
 		ReadbackFromIntermediateBuffer,
 		FReadbackFromIntermediateBufferContext,Context,ReadbackFromIntermediateBufferContext,
 	{
-		RHICopyToResolveTarget(
+		RHICmdList.CopyToResolveTarget(
 			Context.CrashReportResource->GetBuffer(),
 			Context.CrashReportResource->GetReadbackBuffer(),
 			false,
@@ -972,7 +975,7 @@ void FSlateRHIRenderer::MapVirtualScreenBuffer(void** OutImageData)
 		SCOPE_CYCLE_COUNTER(STAT_MapStagingBuffer);
 		int32 UnusedWidth = 0;
 		int32 UnusedHeight = 0;
-		RHIMapStagingSurface(Context.CrashReportResource->GetReadbackBuffer(), *Context.OutData, UnusedWidth, UnusedHeight);
+		RHICmdList.MapStagingSurface(Context.CrashReportResource->GetReadbackBuffer(), *Context.OutData, UnusedWidth, UnusedHeight);
 		Context.CrashReportResource->SwapTargetReadbackBuffer();
 	});
 }
@@ -992,7 +995,7 @@ void FSlateRHIRenderer::UnmapVirtualScreenBuffer()
 		FReadbackFromStagingBufferContext,Context,ReadbackFromStagingBufferContext,
 	{
 		SCOPE_CYCLE_COUNTER(STAT_UnmapStagingBuffer);
-		RHIUnmapStagingSurface(Context.CrashReportResource->GetReadbackBuffer());
+		RHICmdList.UnmapStagingSurface(Context.CrashReportResource->GetReadbackBuffer());
 	});
 }
 
diff --git a/Engine/Source/Runtime/SlateRHIRenderer/Private/SlateRHIRenderer.h b/Engine/Source/Runtime/SlateRHIRenderer/Private/SlateRHIRenderer.h
--- a/Engine/Source/Runtime/SlateRHIRenderer/Private/SlateRHIRenderer.h
+++ b/Engine/Source/Runtime/SlateRHIRenderer/Private/SlateRHIRenderer.h
@@ -40,10 +40,10 @@ public:
 		, UnscaledVirtualScreen(InUnscaledVirtualScreen) {}
 
 	/** FRenderResource Interface.  Called when render resources need to be initialized */
-	virtual void InitDynamicRHI();
+	virtual void InitDynamicRHI() override;
 
 	/** FRenderResource Interface.  Called when render resources need to be released */
-	virtual void ReleaseDynamicRHI();
+	virtual void ReleaseDynamicRHI() override;
 	
 	/** Changes the target readback buffer */
 	void SwapTargetReadbackBuffer() {ReadbackBufferIndex = (ReadbackBufferIndex + 1) % 2;}
@@ -109,8 +109,8 @@ private:
 		bool bFullscreen;
 	
 		/** FRenderResource interface */
-		virtual void InitRHI();
-		virtual void ReleaseRHI();
+		virtual void InitRHI() override;
+		virtual void ReleaseRHI() override;
 
 		FViewportInfo()
 			:	OSWindow(NULL), 
@@ -160,7 +160,7 @@ public:
 	virtual void SetColorVisionDeficiencyType( uint32 Type ) override;
 
 	/** Draws windows from a FSlateDrawBuffer on the render thread */
-	void DrawWindow_RenderThread( const FSlateRHIRenderer::FViewportInfo& ViewportInfo, const FSlateWindowElementList& WindowElementList, bool bLockToVsync );
+	void DrawWindow_RenderThread(FRHICommandListImmediate& RHICmdList, const FSlateRHIRenderer::FViewportInfo& ViewportInfo, const FSlateWindowElementList& WindowElementList, bool bLockToVsync);
 
 
 	/**
diff --git a/Engine/Source/Runtime/SlateRHIRenderer/Private/SlateRHIRenderingPolicy.cpp b/Engine/Source/Runtime/SlateRHIRenderer/Private/SlateRHIRenderingPolicy.cpp
--- a/Engine/Source/Runtime/SlateRHIRenderer/Private/SlateRHIRenderingPolicy.cpp
+++ b/Engine/Source/Runtime/SlateRHIRenderer/Private/SlateRHIRenderingPolicy.cpp
@@ -263,16 +263,13 @@ static FSceneView& CreateSceneView( FSceneViewFamilyContext& ViewFamilyContext,
 	return *View;
 }
 
-void FSlateRHIRenderingPolicy::DrawElements( const FIntPoint& InViewportSize, FSlateBackBuffer& BackBuffer, const FMatrix& ViewProjectionMatrix, const TArray<FSlateRenderBatch>& RenderBatches )
+void FSlateRHIRenderingPolicy::DrawElements(FRHICommandListImmediate& RHICmdList, const FIntPoint& InViewportSize, FSlateBackBuffer& BackBuffer, const FMatrix& ViewProjectionMatrix, const TArray<FSlateRenderBatch>& RenderBatches)
 {
 	SCOPE_CYCLE_COUNTER( STAT_SlateDrawTime );
 
 	// Should only be called by the rendering thread
 	check(IsInRenderingThread());
 
-	//@todo-rco: RHIPacketList
-	FRHICommandList& RHICmdList = FRHICommandList::GetNullRef();
-
 	TSlateElementVertexBuffer<FSlateVertex>& VertexBuffer = VertexBuffers[CurrentBufferIndex];
 	FSlateElementIndexBuffer& IndexBuffer = IndexBuffers[CurrentBufferIndex];
 
@@ -333,7 +330,7 @@ void FSlateRHIRenderingPolicy::DrawElements( const FIntPoint& InViewportSize, FS
 					PixelShader->GetPixelShader(),
 					FGeometryShaderRHIRef());
 
-				RHISetBoundShaderState(BoundShaderState);
+				RHICmdList.SetBoundShaderState(BoundShaderState);
 
 				VertexShader->SetViewProjection(RHICmdList, ViewProjectionMatrix);
 				VertexShader->SetShaderParameters(RHICmdList, ShaderParams.VertexParams);
@@ -449,7 +446,7 @@ void FSlateRHIRenderingPolicy::DrawElements( const FIntPoint& InViewportSize, FS
 						PixelShader->GetPixelShader(),
 						FGeometryShaderRHIRef());
 
-					RHISetBoundShaderState(BoundShaderState);
+					RHICmdList.SetBoundShaderState(BoundShaderState);
 
 					VertexShader->SetShaderParameters(RHICmdList, ShaderParams.VertexParams);
 					PixelShader->SetParameters(RHICmdList, *SceneView, MaterialRenderProxy, Material, 1.0f / DisplayGamma, ShaderParams.PixelParams);
@@ -467,7 +464,7 @@ void FSlateRHIRenderingPolicy::DrawElements( const FIntPoint& InViewportSize, FS
 		else if (RenderBatch.CustomDrawer.IsValid())
 		{
 			// This element is custom and has no Slate geometry.  Tell it to render itself now
-			RenderBatch.CustomDrawer.Pin()->DrawRenderThread( &BackBuffer.GetRenderTargetTexture() );
+			RenderBatch.CustomDrawer.Pin()->DrawRenderThread(RHICmdList, &BackBuffer.GetRenderTargetTexture());
 		}
 
 	}
diff --git a/Engine/Source/Runtime/SlateRHIRenderer/Private/SlateRHIRenderingPolicy.h b/Engine/Source/Runtime/SlateRHIRenderer/Private/SlateRHIRenderingPolicy.h
--- a/Engine/Source/Runtime/SlateRHIRenderer/Private/SlateRHIRenderingPolicy.h
+++ b/Engine/Source/Runtime/SlateRHIRenderer/Private/SlateRHIRenderingPolicy.h
@@ -135,10 +135,10 @@ public:
 	~FSlateElementIndexBuffer();
 
 	/** Initializes the index buffers RHI resource. */
-	virtual void InitDynamicRHI();
+	virtual void InitDynamicRHI() override;
 
 	/** Releases the index buffers RHI resource. */
-	virtual void ReleaseDynamicRHI();
+	virtual void ReleaseDynamicRHI() override;
 
 	/** Returns a friendly name for this buffer. */
 	virtual FString GetFriendlyName() const { return TEXT("SlateElementIndices"); }
@@ -177,11 +177,11 @@ public:
 	~FSlateRHIRenderingPolicy();
 	
 	virtual void UpdateBuffers( const FSlateWindowElementList& WindowElementList ) override;
-	virtual void DrawElements( const FIntPoint& InViewportSize, class FSlateBackBuffer& BackBuffer, const FMatrix& ViewProjectionMatrix, const TArray<FSlateRenderBatch>& RenderBatches );
+	virtual void DrawElements(FRHICommandListImmediate& RHICmdList, const FIntPoint& InViewportSize, class FSlateBackBuffer& BackBuffer, const FMatrix& ViewProjectionMatrix, const TArray<FSlateRenderBatch>& RenderBatches);
 
 	virtual TSharedRef<FSlateFontCache> GetFontCache() override { return FontCache.ToSharedRef(); }
 	virtual TSharedRef<FSlateShaderResourceManager> GetResourceManager() override { return ResourceManager; }
-
+	
 	void InitResources();
 	void ReleaseResources();
 
diff --git a/Engine/Source/Runtime/SlateRHIRenderer/Private/SlateShaders.h b/Engine/Source/Runtime/SlateRHIRenderer/Private/SlateShaders.h
--- a/Engine/Source/Runtime/SlateRHIRenderer/Private/SlateShaders.h
+++ b/Engine/Source/Runtime/SlateRHIRenderer/Private/SlateShaders.h
@@ -22,10 +22,10 @@ public:
 	virtual ~FSlateVertexDeclaration() {}
 
 	/** Initializes the vertex declaration RHI resource */
-	virtual void InitRHI();
+	virtual void InitRHI() override;
 
 	/** Releases the vertex declaration RHI resource */
-	virtual void ReleaseRHI();
+	virtual void ReleaseRHI() override;
 };
 
 /** The slate Vertex shader representation */
diff --git a/Engine/Source/Runtime/StreamingPauseRendering/Private/StreamingPauseRendering.cpp b/Engine/Source/Runtime/StreamingPauseRendering/Private/StreamingPauseRendering.cpp
--- a/Engine/Source/Runtime/StreamingPauseRendering/Private/StreamingPauseRendering.cpp
+++ b/Engine/Source/Runtime/StreamingPauseRendering/Private/StreamingPauseRendering.cpp
@@ -91,7 +91,7 @@ void FStreamingPauseRenderingModule::BeginStreamingPause( FViewport* GameViewpor
 			FViewport*,Viewport,Viewport.Get(),
 			//FIntPoint,InRestoreSize,RestoreSize,
 		{
-			Viewport->EndRenderFrame( false, false );
+			Viewport->EndRenderFrame(RHICmdList, false, false);
 		});
 	}
 
diff --git a/Engine/Source/Runtime/Windows/D3D11RHI/Private/D3D11Commands.cpp b/Engine/Source/Runtime/Windows/D3D11RHI/Private/D3D11Commands.cpp
--- a/Engine/Source/Runtime/Windows/D3D11RHI/Private/D3D11Commands.cpp
+++ b/Engine/Source/Runtime/Windows/D3D11RHI/Private/D3D11Commands.cpp
@@ -1159,54 +1159,6 @@ static D3D11_PRIMITIVE_TOPOLOGY GetD3D11PrimitiveType(uint32 PrimitiveType, bool
 	return D3D11_PRIMITIVE_TOPOLOGY_TRIANGLELIST;
 }
 
-static uint32 GetVertexCountForPrimitiveCount(uint32 NumPrimitives, uint32 PrimitiveType)
-{
-	uint32 VertexCount = 0;
-	switch(PrimitiveType)
-	{
-	case PT_TriangleList: VertexCount = NumPrimitives*3; break;
-	case PT_TriangleStrip: VertexCount = NumPrimitives+2; break;
-	case PT_LineList: VertexCount = NumPrimitives*2; break;
-
-	case PT_1_ControlPointPatchList:
-	case PT_2_ControlPointPatchList:
-	case PT_3_ControlPointPatchList: 
-	case PT_4_ControlPointPatchList: 
-	case PT_5_ControlPointPatchList:
-	case PT_6_ControlPointPatchList:
-	case PT_7_ControlPointPatchList: 
-	case PT_8_ControlPointPatchList: 
-	case PT_9_ControlPointPatchList: 
-	case PT_10_ControlPointPatchList: 
-	case PT_11_ControlPointPatchList: 
-	case PT_12_ControlPointPatchList: 
-	case PT_13_ControlPointPatchList: 
-	case PT_14_ControlPointPatchList: 
-	case PT_15_ControlPointPatchList: 
-	case PT_16_ControlPointPatchList: 
-	case PT_17_ControlPointPatchList: 
-	case PT_18_ControlPointPatchList: 
-	case PT_19_ControlPointPatchList: 
-	case PT_20_ControlPointPatchList: 
-	case PT_21_ControlPointPatchList: 
-	case PT_22_ControlPointPatchList: 
-	case PT_23_ControlPointPatchList: 
-	case PT_24_ControlPointPatchList: 
-	case PT_25_ControlPointPatchList: 
-	case PT_26_ControlPointPatchList: 
-	case PT_27_ControlPointPatchList: 
-	case PT_28_ControlPointPatchList: 
-	case PT_29_ControlPointPatchList: 
-	case PT_30_ControlPointPatchList: 
-	case PT_31_ControlPointPatchList: 
-	case PT_32_ControlPointPatchList: 
-		VertexCount = (PrimitiveType - PT_1_ControlPointPatchList + 1) * NumPrimitives;
-		break;
-	default: UE_LOG(LogD3D11RHI, Fatal,TEXT("Unknown primitive type: %u"),PrimitiveType);
-	};
-
-	return VertexCount;
-}
 
 void FD3D11DynamicRHI::CommitNonComputeShaderConstants()
 {
@@ -1409,7 +1361,7 @@ void FD3D11DynamicRHI::RHIDrawPrimitive(uint32 PrimitiveType,uint32 BaseVertexIn
 	CommitGraphicsResourceTables();
 	CommitNonComputeShaderConstants();
 
-	uint32 VertexCount = RHIGetVertexCountForPrimitiveCount(NumPrimitives,PrimitiveType);
+	uint32 VertexCount = GetVertexCountForPrimitiveCount(NumPrimitives,PrimitiveType);
 
 	GPUProfilingData.RegisterGPUWork(NumPrimitives * NumInstances, VertexCount * NumInstances);
 
@@ -1492,7 +1444,7 @@ void FD3D11DynamicRHI::RHIDrawIndexedPrimitive(FIndexBufferRHIParamRef IndexBuff
 	uint32 SizeFormat = sizeof(DXGI_FORMAT);
 	const DXGI_FORMAT Format = (IndexBuffer->GetStride() == sizeof(uint16) ? DXGI_FORMAT_R16_UINT : DXGI_FORMAT_R32_UINT);
 
-	uint32 IndexCount = RHIGetVertexCountForPrimitiveCount(NumPrimitives,PrimitiveType);
+	uint32 IndexCount = GetVertexCountForPrimitiveCount(NumPrimitives,PrimitiveType);
 
 	// Verify that we are not trying to read outside the index buffer range
 	// test is an optimized version of: StartIndex + IndexCount <= IndexBuffer->GetSize() / IndexBuffer->GetStride() 
@@ -1668,6 +1620,8 @@ void FD3D11DynamicRHI::RHIClear(bool bClearColor,const FLinearColor& Color,bool
 
 void FD3D11DynamicRHI::RHIClearMRT(bool bClearColor,int32 NumClearColors,const FLinearColor* ClearColorArray,bool bClearDepth,float Depth,bool bClearStencil,uint32 Stencil, FIntRect ExcludeRect)
 {
+	FRHICommandListImmediate& RHICmdList = FRHICommandListExecutor::GetRecursiveRHICommandList();
+
 	GRHICommandList.Verify();
 
 	// Helper struct to record and restore device states RHIClearMRT modifies.
@@ -1987,9 +1941,7 @@ void FD3D11DynamicRHI::RHIClearMRT(bool bClearColor,int32 NumClearColors,const F
 			TShaderMapRef<TOneColorPixelShaderMRT<8> > MRTPixelShader(GetGlobalShaderMap());
 			PixelShader = *MRTPixelShader;
 		}
-		//@todo-rco: RHIPacketList
-		FRHICommandList& RHICmdList = FRHICommandList::GetNullRef();
-
+		
 		SetGlobalBoundShaderState(RHICmdList, GD3D11ClearMRTBoundShaderState[FMath::Max(BoundRenderTargets.GetNumActiveTargets() - 1, 0)], GD3D11Vector4VertexDeclaration.VertexDeclarationRHI, *VertexShader, PixelShader);
 		FLinearColor ShaderClearColors[MaxSimultaneousRenderTargets];
 		FMemory::MemZero(ShaderClearColors);
@@ -2047,6 +1999,7 @@ void FD3D11DynamicRHI::RHIClearMRT(bool bClearColor,int32 NumClearColors,const F
 				DrawPrimitiveUP(RHICmdList, PT_TriangleStrip, 2, Vertices, sizeof(Vertices[0]));
 			}
 		}
+		RHICmdList.Flush(); // always call flush with GetRecursiveRHICommandList, recursive use of the RHI is hazardous
 
 		if (bChangedDepthStencilTarget)
 		{
@@ -2333,6 +2286,106 @@ void FD3D11DynamicRHI::RHIExecuteCommandList(FRHICommandList* CmdList)
 				RHICmd->Texture->Release();
 			}
 			break;
+		case ERCT_SetShaderResourceViewParameter:
+			{
+				auto* RHICmd = Iter.NextCommand<FRHICommandSetShaderResourceViewParameter>();
+				{
+					uint32 Start = FPlatformTime::Cycles();
+
+					FD3D11ShaderResourceView*  SRV = (FD3D11ShaderResourceView*)RHICmd->SRV;
+
+					FD3D11BaseShaderResource* Resource = SRV->Resource;
+					ID3D11ShaderResourceView* D3D11SRV = SRV->View;
+
+					switch (RHICmd->ShaderFrequency)
+					{
+					case SF_Vertex:
+						VALIDATE_BOUND_SHADER((FVertexShaderRHIParamRef)RHICmd->Shader);
+						SetShaderResourceView<SF_Vertex>(Resource, D3D11SRV, RHICmd->SamplerIndex);
+						break;
+					case SF_Hull:
+						VALIDATE_BOUND_SHADER((FHullShaderRHIParamRef)RHICmd->Shader);
+						SetShaderResourceView<SF_Hull>(Resource, D3D11SRV, RHICmd->SamplerIndex);
+						break;
+					case SF_Domain:
+						VALIDATE_BOUND_SHADER((FDomainShaderRHIParamRef)RHICmd->Shader);
+						SetShaderResourceView<SF_Domain>(Resource, D3D11SRV, RHICmd->SamplerIndex);
+						break;
+					case SF_Geometry:
+						VALIDATE_BOUND_SHADER((FGeometryShaderRHIParamRef)RHICmd->Shader);
+						SetShaderResourceView<SF_Geometry>(Resource, D3D11SRV, RHICmd->SamplerIndex);
+						break;
+					case SF_Pixel:
+						VALIDATE_BOUND_SHADER((FPixelShaderRHIParamRef)RHICmd->Shader);
+						SetShaderResourceView<SF_Pixel>(Resource, D3D11SRV, RHICmd->SamplerIndex);
+						break;
+					case SF_Compute:
+						VALIDATE_BOUND_SHADER((FComputeShaderRHIParamRef)RHICmd->Shader);
+						SetShaderResourceView<SF_Compute>(Resource, D3D11SRV, RHICmd->SamplerIndex);
+						break;
+					default: check(0); break;
+					}
+				}
+				RHICmd->Shader->Release();
+				RHICmd->SRV->Release();
+			}
+			break;
+		case ERCT_SetUAVParameter:
+			{
+				auto* RHICmd = Iter.NextCommand<FRHICommandSetUAVParameter>();
+				{
+					uint32 Start = FPlatformTime::Cycles();
+
+					FD3D11UnorderedAccessView*  UAV = (FD3D11UnorderedAccessView*)RHICmd->UAV;
+
+					FD3D11BaseShaderResource* Resource = UAV->Resource;
+					ID3D11UnorderedAccessView* D3D11UAV = UAV->View;
+
+
+					ConditionalClearShaderResource(Resource);
+
+					uint32 InitialCount = -1;
+
+					switch (RHICmd->ShaderFrequency)
+					{
+					case SF_Compute:
+						//VALIDATE_BOUND_SHADER((FComputeShaderRHIParamRef)RHICmd->Shader);
+						Direct3DDeviceIMContext->CSSetUnorderedAccessViews(RHICmd->UAVIndex, 1, &D3D11UAV, &InitialCount);
+						break;
+					default: check(0); break;
+					}
+				}
+				RHICmd->Shader->Release();
+				RHICmd->UAV->Release();
+			}
+			break;
+		case ERCT_SetUAVParameter_IntialCount:
+			{
+				auto* RHICmd = Iter.NextCommand<FRHICommandSetUAVParameter_IntialCount>();
+				{
+					uint32 Start = FPlatformTime::Cycles();
+
+					FD3D11UnorderedAccessView*  UAV = (FD3D11UnorderedAccessView*)RHICmd->UAV;
+
+					FD3D11BaseShaderResource* Resource = UAV->Resource;
+					ID3D11UnorderedAccessView* D3D11UAV = UAV->View;
+
+
+					ConditionalClearShaderResource(Resource);
+
+					switch (RHICmd->ShaderFrequency)
+					{
+					case SF_Compute:
+						//VALIDATE_BOUND_SHADER((FComputeShaderRHIParamRef)RHICmd->Shader);
+						Direct3DDeviceIMContext->CSSetUnorderedAccessViews(RHICmd->UAVIndex, 1, &D3D11UAV, &RHICmd->InitialCount);
+						break;
+					default: check(0); break;
+					}
+				}
+				RHICmd->Shader->Release();
+				RHICmd->UAV->Release();
+			}
+			break;
 		case ERCT_DrawPrimitive:
 			{
 				auto* RHICmd = Iter.NextCommand<FRHICommandDrawPrimitive>();
@@ -2341,7 +2394,7 @@ void FD3D11DynamicRHI::RHIExecuteCommandList(FRHICommandList* CmdList)
 
 					CommitGraphicsResourceTables();
 					CommitNonComputeShaderConstants();
-					const uint32 VertexCount = RHIGetVertexCountForPrimitiveCount(RHICmd->NumPrimitives, RHICmd->PrimitiveType);
+					const uint32 VertexCount = GetVertexCountForPrimitiveCount(RHICmd->NumPrimitives, RHICmd->PrimitiveType);
 
 					GPUProfilingData.RegisterGPUWork(RHICmd->NumPrimitives * RHICmd->NumInstances, VertexCount * RHICmd->NumInstances);
 					StateCache.SetPrimitiveTopology(GetD3D11PrimitiveType(RHICmd->PrimitiveType, bUsingTessellation));
@@ -2376,7 +2429,7 @@ void FD3D11DynamicRHI::RHIExecuteCommandList(FRHICommandList* CmdList)
 					uint32 SizeFormat = sizeof(DXGI_FORMAT);
 					const DXGI_FORMAT Format = (IndexBuffer->GetStride() == sizeof(uint16) ? DXGI_FORMAT_R16_UINT : DXGI_FORMAT_R32_UINT);
 
-					uint32 IndexCount = RHIGetVertexCountForPrimitiveCount(RHICmd->NumPrimitives, RHICmd->PrimitiveType);
+					uint32 IndexCount = GetVertexCountForPrimitiveCount(RHICmd->NumPrimitives, RHICmd->PrimitiveType);
 
 					// Verify that we are not trying to read outside the index buffer range
 					checkf(RHICmd->StartIndex + IndexCount <= IndexBuffer->GetSize() / IndexBuffer->GetStride(), 
diff --git a/Engine/Source/Runtime/Windows/D3D11RHI/Private/D3D11RHI.cpp b/Engine/Source/Runtime/Windows/D3D11RHI/Private/D3D11RHI.cpp
--- a/Engine/Source/Runtime/Windows/D3D11RHI/Private/D3D11RHI.cpp
+++ b/Engine/Source/Runtime/Windows/D3D11RHI/Private/D3D11RHI.cpp
@@ -182,8 +182,8 @@ FGlobalBoundShaderState LongGPUTaskBoundShaderState;
 
 void FD3D11DynamicRHI::IssueLongGPUTask()
 {
-	//@todo-rco: RHIPacketList
-	FRHICommandList& RHICmdList = FRHICommandList::GetNullRef();
+	
+	FRHICommandListImmediate& RHICmdList = FRHICommandListExecutor::GetRecursiveRHICommandList();
 
 	if (GRHIFeatureLevel >= ERHIFeatureLevel::SM3)
 	{
@@ -224,6 +224,7 @@ void FD3D11DynamicRHI::IssueLongGPUTask()
 			DrawPrimitiveUP(RHICmdList, PT_TriangleStrip, 2, Vertices, sizeof(Vertices[0]));
 		}
 	}
+	RHICmdList.Flush(); // always call flush with GetRecursiveRHICommandList, recursive use of the RHI is hazardous
 }
 
 void FD3DGPUProfiler::BeginFrame(FD3D11DynamicRHI* InRHI)
diff --git a/Engine/Source/Runtime/Windows/D3D11RHI/Private/D3D11RenderTarget.cpp b/Engine/Source/Runtime/Windows/D3D11RHI/Private/D3D11RenderTarget.cpp
--- a/Engine/Source/Runtime/Windows/D3D11RHI/Private/D3D11RenderTarget.cpp
+++ b/Engine/Source/Runtime/Windows/D3D11RHI/Private/D3D11RenderTarget.cpp
@@ -172,8 +172,8 @@ void FD3D11DynamicRHI::ResolveTextureUsingShader(
 	typename TPixelShader::FParameter PixelShaderParameter
 	)
 {
-	//@todo-rco: RHIPacketList
-	FRHICommandList& RHICmdList = FRHICommandList::GetNullRef();
+	
+	FRHICommandListImmediate& RHICmdList = FRHICommandListExecutor::GetRecursiveRHICommandList();
 
 	// Save the current viewport so that it can be restored
 	D3D11_VIEWPORT SavedViewport;
@@ -181,8 +181,9 @@ void FD3D11DynamicRHI::ResolveTextureUsingShader(
 	StateCache.GetViewports(&NumSavedViewports,&SavedViewport);
 
 	// No alpha blending, no depth tests or writes, no stencil tests or writes, no backface culling.
-	RHISetBlendState(TStaticBlendState<>::GetRHI(),FLinearColor::White);
-	RHISetRasterizerState(TStaticRasterizerState<FM_Solid,CM_None>::GetRHI());
+	RHICmdList.SetBlendState(TStaticBlendState<>::GetRHI(),FLinearColor::White);
+	RHICmdList.SetRasterizerState(TStaticRasterizerState<FM_Solid, CM_None>::GetRHI());
+	RHICmdList.Flush(); // always call flush with GetRecursiveRHICommandList, recursive use of the RHI is hazardous
 
 	// Make sure the destination is not bound as a shader resource.
 	if (DestTexture)
@@ -209,7 +210,8 @@ void FD3D11DynamicRHI::ResolveTextureUsingShader(
 			Direct3DDeviceContext->ClearDepthStencilView(DestTextureDSV,D3D11_CLEAR_DEPTH | D3D11_CLEAR_STENCIL,0,0);
 		}
 
-		RHISetDepthStencilState(TStaticDepthStencilState<true,CF_Always>::GetRHI(),0);
+		RHICmdList.SetDepthStencilState(TStaticDepthStencilState<true,CF_Always>::GetRHI(),0);
+		RHICmdList.Flush(); // always call flush with GetRecursiveRHICommandList, recursive use of the RHI is hazardous
 
 		// Write to the dest texture as a depth-stencil target.
 		ID3D11RenderTargetView* NullRTV = NULL;
@@ -226,13 +228,14 @@ void FD3D11DynamicRHI::ResolveTextureUsingShader(
 			Direct3DDeviceContext->ClearRenderTargetView(DestTextureRTV,(float*)&ClearColor);
 		}
 
-		RHISetDepthStencilState(TStaticDepthStencilState<false,CF_Always>::GetRHI(),0);
+		RHICmdList.SetDepthStencilState(TStaticDepthStencilState<false,CF_Always>::GetRHI(),0);
+		RHICmdList.Flush(); // always call flush with GetRecursiveRHICommandList, recursive use of the RHI is hazardous
 
 		// Write to the dest surface as a render target.
 		Direct3DDeviceContext->OMSetRenderTargets(1,&DestTextureRTV,NULL);
 	}
 
-	RHISetViewport(0.0f, 0.0f, 0.0f, ResolveTargetDesc.Width, ResolveTargetDesc.Height, 1.0f  );
+	RHICmdList.SetViewport(0.0f, 0.0f, 0.0f, ResolveTargetDesc.Width, ResolveTargetDesc.Height, 1.0f  );
 
 
 	// Generate the vertices used to copy from the source surface to the destination surface.
@@ -253,6 +256,7 @@ void FD3D11DynamicRHI::ResolveTextureUsingShader(
 	SetGlobalBoundShaderState(RHICmdList, ResolveBoundShaderState, GScreenVertexDeclaration.VertexDeclarationRHI, *ResolveVertexShader, *ResolvePixelShader);
 
 	ResolvePixelShader->SetParameters(RHICmdList, Direct3DDeviceContext,PixelShaderParameter);
+	RHICmdList.Flush(); // always call flush with GetRecursiveRHICommandList, recursive use of the RHI is hazardous
 
 	// Set the source texture.
 	const uint32 TextureIndex = ResolvePixelShader->UnresolvedSurface.GetBaseIndex();
@@ -282,6 +286,7 @@ void FD3D11DynamicRHI::ResolveTextureUsingShader(
 	Vertices[3].UV.Y       = MaxV;
 
 	DrawPrimitiveUP(RHICmdList, PT_TriangleStrip, 2, Vertices, sizeof(Vertices[0]));
+	RHICmdList.Flush(); // always call flush with GetRecursiveRHICommandList, recursive use of the RHI is hazardous
 
 	if (SourceTexture)
 	{
