UE4 -fix nonunity build

diff --git a/Engine/Source/Runtime/SlateRHIRenderer/Private/SlateRHIRenderingPolicy.cpp b/Engine/Source/Runtime/SlateRHIRenderer/Private/SlateRHIRenderingPolicy.cpp
--- a/Engine/Source/Runtime/SlateRHIRenderer/Private/SlateRHIRenderingPolicy.cpp
+++ b/Engine/Source/Runtime/SlateRHIRenderer/Private/SlateRHIRenderingPolicy.cpp
@@ -4,6 +4,7 @@
 #include "RenderingPolicy.h"
 #include "SlateRHIRenderingPolicy.h"
 #include "SlateShaders.h"
+#include "SlateMaterialShader.h"
 #include "SlateRHIResourceManager.h"
 #include "TileRendering.h"
 #include "PreviewScene.h"
