Fix race condition with custom Slate draw elements

diff --git a/Engine/Source/Runtime/SlateRHIRenderer/Private/SlateRHIRenderingPolicy.cpp b/Engine/Source/Runtime/SlateRHIRenderer/Private/SlateRHIRenderingPolicy.cpp
--- a/Engine/Source/Runtime/SlateRHIRenderer/Private/SlateRHIRenderingPolicy.cpp
+++ b/Engine/Source/Runtime/SlateRHIRenderer/Private/SlateRHIRenderingPolicy.cpp
@@ -483,10 +483,14 @@ void FSlateRHIRenderingPolicy::DrawElements(FRHICommandListImmediate& RHICmdList
 			}
 
 		}
-		else if (RenderBatch.CustomDrawer.IsValid())
+		else
 		{
-			// This element is custom and has no Slate geometry.  Tell it to render itself now
-			RenderBatch.CustomDrawer.Pin()->DrawRenderThread(RHICmdList, &BackBuffer.GetRenderTargetTexture());
+			TSharedPtr<ICustomSlateElement, ESPMode::ThreadSafe> CustomDrawer = RenderBatch.CustomDrawer.Pin();
+			if (CustomDrawer.IsValid())
+			{
+				// This element is custom and has no Slate geometry.  Tell it to render itself now
+				CustomDrawer->DrawRenderThread(RHICmdList, &BackBuffer.GetRenderTargetTexture());
+			}
 		}
 
 	}
