Window morphers now work again in the editor

#ttp 340320 - REGRESSION: EDITOR: Menu windows do not open when 'Enable Window Animations' is enabled

#branch UE4

#change bHasEverBeenDrawn wasn't being marked as true in the SlateRenderer (just in the StandaloneRenderer), which prevented any morpher from updating in SWindow::Tick

diff --git a/Engine/Source/Runtime/SlateRHIRenderer/Private/SlateRHIRenderer.cpp b/Engine/Source/Runtime/SlateRHIRenderer/Private/SlateRHIRenderer.cpp
--- a/Engine/Source/Runtime/SlateRHIRenderer/Private/SlateRHIRenderer.cpp
+++ b/Engine/Source/Runtime/SlateRHIRenderer/Private/SlateRHIRenderer.cpp
@@ -639,6 +639,12 @@ void FSlateRHIRenderer::DrawWindows_Private( FSlateDrawBuffer& WindowDrawBuffer
 						Params.Renderer->DrawWindow_RenderThread(RHICmdList, *Params.ViewportInfo, *Params.WindowElementList, Params.bLockToVsync);
 					});
 
+					// Mark this window as 'drawn'
+					if (!Window->HasEverBeenDrawn())
+					{
+						Window->MarkAsDrawn();
+					}
+
 					SlateWindowRendered.Broadcast( *Params.SlateWindow, &ViewInfo->ViewportRHI );
 
 					if ( bTakingAScreenShot )
