Rollback //depot/UE4/Engine/Source/Runtime/SlateRHIRenderer/Private/SlateRHIRenderer.cpp to revision 36, just while I investigate...

diff --git a/Engine/Source/Runtime/SlateRHIRenderer/Private/SlateRHIRenderer.cpp b/Engine/Source/Runtime/SlateRHIRenderer/Private/SlateRHIRenderer.cpp
--- a/Engine/Source/Runtime/SlateRHIRenderer/Private/SlateRHIRenderer.cpp
+++ b/Engine/Source/Runtime/SlateRHIRenderer/Private/SlateRHIRenderer.cpp
@@ -639,12 +639,6 @@ void FSlateRHIRenderer::DrawWindows_Private( FSlateDrawBuffer& WindowDrawBuffer
 						Params.Renderer->DrawWindow_RenderThread(RHICmdList, *Params.ViewportInfo, *Params.WindowElementList, Params.bLockToVsync);
 					});
 
-					// Mark this window as 'drawn'
-					if (!Window->HasEverBeenDrawn())
-					{
-						Window->MarkAsDrawn();
-					}
-
 					SlateWindowRendered.Broadcast( *Params.SlateWindow, &ViewInfo->ViewportRHI );
 
 					if ( bTakingAScreenShot )
