Slate Resource Management changes to support 3D

diff --git a/Engine/Source/Runtime/Engine/Public/Slate/SlateTextures.h b/Engine/Source/Runtime/Engine/Public/Slate/SlateTextures.h
--- a/Engine/Source/Runtime/Engine/Public/Slate/SlateTextures.h
+++ b/Engine/Source/Runtime/Engine/Public/Slate/SlateTextures.h
@@ -117,28 +117,6 @@ private:
 	uint32 Height;
 };
 
-
-class FSlateMaterial : public FSlateShaderResource
-{
-public:
-	FSlateMaterial( FMaterialRenderProxy* InProxy,  const FVector2D& InImageSize )
-		: RenderProxy( InProxy )
-		, Width( FMath::RoundToInt( InImageSize.X ) )
-		, Height( FMath::RoundToInt( InImageSize.Y ) )
-	{}
-
-	virtual uint32 GetWidth() const { return Width; } 
-	virtual uint32 GetHeight() const { return Height; } 
-
-	virtual void SetResource( FMaterialRenderProxy* Proxy ) { RenderProxy = Proxy; }
-	virtual FMaterialRenderProxy* GetRenderProxy() const { return RenderProxy; }
-	virtual ESlateShaderResource::Type GetType() const override { return ESlateShaderResource::Material; }
-private:
-	FMaterialRenderProxy* RenderProxy;
-	uint32 Width;
-	uint32 Height;
-};
-
 class FSlateTextureRenderTarget2DResource : public FTextureRenderTargetResource
 {
 public:
diff --git a/Engine/Source/Runtime/SlateCore/Public/Textures/SlateShaderResource.h b/Engine/Source/Runtime/SlateCore/Public/Textures/SlateShaderResource.h
--- a/Engine/Source/Runtime/SlateCore/Public/Textures/SlateShaderResource.h
+++ b/Engine/Source/Runtime/SlateCore/Public/Textures/SlateShaderResource.h
@@ -11,7 +11,10 @@ namespace ESlateShaderResource
 	enum Type
 	{
 		/** Texture resource. */
-		Texture,
+		NativeTexture,
+
+		/** UTexture object resource */
+		TextureObject,
 
 		/** Material resource. */
 		Material
@@ -124,10 +127,9 @@ public:
 public:
 
 	// FSlateShaderResource interface
-
 	virtual ESlateShaderResource::Type GetType() const override
 	{
-		return ESlateShaderResource::Texture;
+		return ESlateShaderResource::NativeTexture;
 	}
 
 protected:
diff --git a/Engine/Source/Runtime/SlateRHIRenderer/Private/SlateMaterialResource.cpp b/Engine/Source/Runtime/SlateRHIRenderer/Private/SlateMaterialResource.cpp
new file mode 100644
--- /dev/null
+++ b/Engine/Source/Runtime/SlateRHIRenderer/Private/SlateMaterialResource.cpp
@@ -0,0 +1,25 @@
+// Copyright 1998-2014 Epic Games, Inc. All Rights Reserved.
+
+#include "SlateRHIRendererPrivatePCH.h"
+#include "SlateMaterialResource.h"
+
+
+FSlateMaterialResource::FSlateMaterialResource(UMaterialInterface& InMaterial, const FVector2D& InImageSize)
+	: MaterialObject( &InMaterial )
+	, RenderProxy( InMaterial.GetRenderProxy(false) )
+	, SlateProxy( new FSlateShaderResourceProxy )
+	, Width(FMath::RoundToInt(InImageSize.X))
+	, Height(FMath::RoundToInt(InImageSize.Y))
+{
+	SlateProxy->ActualSize = InImageSize.IntPoint();
+	SlateProxy->Resource = this;
+}
+
+FSlateMaterialResource::~FSlateMaterialResource()
+{
+	if(SlateProxy)
+	{
+		delete SlateProxy;
+	}
+}
+
diff --git a/Engine/Source/Runtime/SlateRHIRenderer/Private/SlateMaterialResource.h b/Engine/Source/Runtime/SlateRHIRenderer/Private/SlateMaterialResource.h
new file mode 100644
--- /dev/null
+++ b/Engine/Source/Runtime/SlateRHIRenderer/Private/SlateMaterialResource.h
@@ -0,0 +1,37 @@
+// Copyright 1998-2014 Epic Games, Inc. All Rights Reserved.
+
+#pragma once
+
+/**
+ * A resource for rendering a UMaterial in Slate
+ */
+class FSlateMaterialResource : public FSlateShaderResource
+{
+public:
+	FSlateMaterialResource(UMaterialInterface& InMaterialResource, const FVector2D& InImageSize);
+	~FSlateMaterialResource();
+
+	virtual uint32 GetWidth() const { return Width; }
+	virtual uint32 GetHeight() const { return Height; }
+	virtual ESlateShaderResource::Type GetType() const override { return ESlateShaderResource::Material; }
+
+	/** Updates the rendering resource */
+	void UpdateRenderResource(FMaterialRenderProxy* Proxy) { RenderProxy = Proxy; }
+
+	/** @return The material render proxy */
+	FMaterialRenderProxy* GetRenderProxy() const { return RenderProxy; }
+
+	/** @return the material object */
+	UMaterialInterface* GetMaterialObject() const { return MaterialObject; }
+
+public:
+	class UMaterialInterface* MaterialObject;
+	/** Rendering thread version of the material object */
+	FMaterialRenderProxy* RenderProxy;
+	/** Slate proxy used for batching the material */
+	FSlateShaderResourceProxy* SlateProxy;
+
+	uint32 Width;
+	uint32 Height;
+};
+
diff --git a/Engine/Source/Runtime/SlateRHIRenderer/Private/SlateNativeTextureResource.cpp b/Engine/Source/Runtime/SlateRHIRenderer/Private/SlateNativeTextureResource.cpp
new file mode 100644
--- /dev/null
+++ b/Engine/Source/Runtime/SlateRHIRenderer/Private/SlateNativeTextureResource.cpp
@@ -0,0 +1,27 @@
+// Copyright 1998-2014 Epic Games, Inc. All Rights Reserved.
+
+#include "SlateRHIRendererPrivatePCH.h"
+#include "SlateNativeTextureResource.h"
+
+TSharedPtr<FSlateDynamicTextureResource> FSlateDynamicTextureResource::NullResource = MakeShareable( new FSlateDynamicTextureResource( NULL ) );
+
+FSlateDynamicTextureResource::FSlateDynamicTextureResource(FSlateTexture2DRHIRef* ExistingTexture)
+	: Proxy(new FSlateShaderResourceProxy)
+	, RHIRefTexture(ExistingTexture != NULL ? ExistingTexture : new FSlateTexture2DRHIRef(NULL, 0, 0))
+{
+	Proxy->Resource = RHIRefTexture;
+}
+
+
+FSlateDynamicTextureResource::~FSlateDynamicTextureResource()
+{
+	if(Proxy)
+	{
+		delete Proxy;
+	}
+
+	if(RHIRefTexture)
+	{
+		delete RHIRefTexture;
+	}
+}
diff --git a/Engine/Source/Runtime/SlateRHIRenderer/Private/SlateNativeTextureResource.h b/Engine/Source/Runtime/SlateRHIRenderer/Private/SlateNativeTextureResource.h
new file mode 100644
--- /dev/null
+++ b/Engine/Source/Runtime/SlateRHIRenderer/Private/SlateNativeTextureResource.h
@@ -0,0 +1,17 @@
+// Copyright 1998-2014 Epic Games, Inc. All Rights Reserved.
+
+#pragma once
+
+
+/** Represents a dynamic texture resource for rendering in Slate*/
+class FSlateDynamicTextureResource
+{
+public:
+	static TSharedPtr<FSlateDynamicTextureResource> NullResource;
+
+	FSlateDynamicTextureResource(FSlateTexture2DRHIRef* ExistingTexture);
+	~FSlateDynamicTextureResource();
+
+	FSlateShaderResourceProxy* Proxy;
+	FSlateTexture2DRHIRef* RHIRefTexture;
+};
\ No newline at end of file
diff --git a/Engine/Source/Runtime/SlateRHIRenderer/Private/SlateRHIRenderer.cpp b/Engine/Source/Runtime/SlateRHIRenderer/Private/SlateRHIRenderer.cpp
--- a/Engine/Source/Runtime/SlateRHIRenderer/Private/SlateRHIRenderer.cpp
+++ b/Engine/Source/Runtime/SlateRHIRenderer/Private/SlateRHIRenderer.cpp
@@ -3,6 +3,7 @@
 #include "SlateRHIRendererPrivatePCH.h"
 #include "StereoRendering.h"
 #include "Runtime/Engine/Public/Features/ILiveStreamingService.h"
+#include "SlateNativeTextureResource.h"
 
 DECLARE_CYCLE_STAT(TEXT("Map Staging Buffer"),STAT_MapStagingBuffer,STATGROUP_CrashTracker);
 DECLARE_CYCLE_STAT(TEXT("Generate Capture Buffer"),STAT_GenerateCaptureBuffer,STATGROUP_CrashTracker);
@@ -12,27 +13,6 @@ DECLARE_CYCLE_STAT(TEXT("Unmap Staging Buffer"),STAT_UnmapStagingBuffer,STATGROU
 // Defines the maximum size that a slate viewport will create
 #define MAX_VIEWPORT_SIZE 16384
 
-static FMatrix CreateProjectionMatrix( uint32 Width, uint32 Height )
-{
-	
-	float PixelOffset = GPixelCenterOffset;
-
-	// Create ortho projection matrix
-	const float Left = PixelOffset;
-	const float Right = Left+Width;
-	const float Top = PixelOffset;
-	const float Bottom = Top+Height;
-	const float ZNear = -100.0f;
-	const float ZFar = 100.0f;
-	return AdjustProjectionMatrixForRHI(
-		FMatrix(
-			FPlane(2.0f/(Right-Left),			0,							0,					0 ),
-			FPlane(0,							2.0f/(Top-Bottom),			0,					0 ),
-			FPlane(0,							0,							1/(ZNear-ZFar),		0 ),
-			FPlane((Left+Right)/(Left-Right),	(Top+Bottom)/(Bottom-Top),	ZNear/(ZNear-ZFar), 1 )
-			)
-		);
-}
 
 
 void FSlateCrashReportResource::InitDynamicRHI()
@@ -120,12 +100,18 @@ void FSlateRHIRenderer::FViewportInfo::RecreateDepthBuffer_RenderThread()
 	}
 }
 
-FSlateRHIRenderer::FSlateRHIRenderer()
+
+
+FSlateRHIRenderer::FSlateRHIRenderer( TSharedPtr<FSlateRHIResourceManager> InResourceManager, TSharedPtr<FSlateFontCache> InFontCache, TSharedPtr<FSlateFontMeasure> InFontMeasure )
 #if USE_MAX_DRAWBUFFERS
 	: EnqueuedWindowDrawBuffer(NULL)
 	, FreeBufferIndex(1)
 #endif
 {
+	ResourceManager = InResourceManager;
+	FontCache = InFontCache;
+	FontMeasure = InFontMeasure;
+
 	CrashTrackerResource = NULL;
 	ViewMatrix = FMatrix(	FPlane(1,	0,	0,	0),
 							FPlane(0,	1,	0,	0),
@@ -140,41 +126,32 @@ FSlateRHIRenderer::~FSlateRHIRenderer()
 {
 }
 
-class FSlateRHIFontAtlasFactory : public ISlateFontAtlasFactory
+FMatrix FSlateRHIRenderer::CreateProjectionMatrix(uint32 Width, uint32 Height)
 {
-public:
-	FSlateRHIFontAtlasFactory()
-	{
-		/** Size of each font texture, width and height */
-		AtlasSize = 1024;
-		if( !GIsEditor && GConfig )
-		{
-			GConfig->GetInt( TEXT("SlateRenderer"), TEXT("FontAtlasSize"), AtlasSize, GEngineIni );
-			AtlasSize = FMath::Clamp( AtlasSize, 0, 2048 );
-		}
-	}
-
-	virtual ~FSlateRHIFontAtlasFactory()
-	{
-	}
 
+	float PixelOffset = GPixelCenterOffset;
 
-	virtual TSharedRef<FSlateFontAtlas> CreateFontAtlas() const override
-	{
-		return MakeShareable( new FSlateFontAtlasRHI( AtlasSize, AtlasSize ) );
-	}
-private:
-	int32 AtlasSize;
-};
+	// Create ortho projection matrix
+	const float Left = PixelOffset;
+	const float Right = Left + Width;
+	const float Top = PixelOffset;
+	const float Bottom = Top + Height;
+	const float ZNear = -100.0f;
+	const float ZFar = 100.0f;
+	return AdjustProjectionMatrixForRHI(
+		FMatrix(
+		FPlane(2.0f / (Right - Left), 0, 0, 0),
+		FPlane(0, 2.0f / (Top - Bottom), 0, 0),
+		FPlane(0, 0, 1 / (ZNear - ZFar), 0),
+		FPlane((Left + Right) / (Left - Right), (Top + Bottom) / (Bottom - Top), ZNear / (ZNear - ZFar), 1)
+		)
+		);
+}
 
 void FSlateRHIRenderer::Initialize()
 {
-	ResourceManager = MakeShareable( new FSlateRHIResourceManager );
 	LoadUsedTextures();
 
-	FontCache = MakeShareable( new FSlateFontCache( MakeShareable( new FSlateRHIFontAtlasFactory ) ) );
-	FontMeasure = FSlateFontMeasure::Create( FontCache.ToSharedRef() );
-
 	RenderingPolicy = MakeShareable( new FSlateRHIRenderingPolicy( FontCache, ResourceManager.ToSharedRef() ) ); 
 
 	ElementBatcher = MakeShareable( new FSlateElementBatcher( RenderingPolicy.ToSharedRef() ) );
@@ -453,7 +430,6 @@ void FSlateRHIRenderer::DrawWindow_RenderThread(FRHICommandListImmediate& RHICmd
 			RenderingPolicy->DrawElements
 			(
 				RHICmdList,
-				FIntPoint(ViewportWidth, ViewportHeight),
 				BackBufferTarget,
 				ViewMatrix*ViewportInfo.ProjectionMatrix,
 				WindowElementList.GetRenderBatches()
@@ -953,7 +929,7 @@ void FSlateRHIRenderer::CopyWindowsToVirtualScreenBuffer(const TArray<FString>&
 			FTexture2DRHIRef UnusedTargetTexture;
 			FSlateBackBuffer UnusedTarget( UnusedTargetTexture, FIntPoint::ZeroValue );
 
-			Context.RenderPolicy->DrawElements(RHICmdList, Context.ViewportSize, UnusedTarget, CreateProjectionMatrix(Context.ViewportSize.X, Context.ViewportSize.Y), Context.SlateElementList->GetRenderBatches());
+			Context.RenderPolicy->DrawElements(RHICmdList, UnusedTarget, CreateProjectionMatrix(Context.ViewportSize.X, Context.ViewportSize.Y), Context.SlateElementList->GetRenderBatches());
 		}
 	});
 
@@ -1027,7 +1003,20 @@ void FSlateRHIRenderer::UnmapVirtualScreenBuffer()
 FIntPoint FSlateRHIRenderer::GenerateDynamicImageResource(const FName InTextureName)
 {
 	check( IsInGameThread() );
-	TSharedPtr<FDynamicTextureResource> TextureResource = ResourceManager->MakeDynamicTextureResource(false, true, InTextureName.ToString(), InTextureName, NULL);
+
+	uint32 Width = 0;
+	uint32 Height = 0;
+	TArray<uint8> RawData;
+
+	// Load the image from disk
+	bool bSucceeded = ResourceManager->LoadTexture(InTextureName, InTextureName.ToString(), Width, Height, RawData);
+
+	TSharedPtr<FSlateDynamicTextureResource> TextureResource;
+	if( bSucceeded )
+	{
+		TextureResource = ResourceManager->MakeDynamicTextureResource(InTextureName, Width, Height, RawData);
+	}
+
 	return TextureResource.IsValid() ? TextureResource->Proxy->ActualSize : FIntPoint( 0, 0 );
 }
 
@@ -1035,7 +1024,7 @@ bool FSlateRHIRenderer::GenerateDynamicImageResource( FName ResourceName, uint32
 {
 	check( IsInGameThread() );
 
-	TSharedPtr<FDynamicTextureResource> TextureResource = ResourceManager->MakeDynamicTextureResource( ResourceName, Width, Height, Bytes );
+	TSharedPtr<FSlateDynamicTextureResource> TextureResource = ResourceManager->MakeDynamicTextureResource( ResourceName, Width, Height, Bytes );
 	return TextureResource.IsValid();
 }
 
diff --git a/Engine/Source/Runtime/SlateRHIRenderer/Private/SlateRHIRenderer.h b/Engine/Source/Runtime/SlateRHIRenderer/Private/SlateRHIRenderer.h
--- a/Engine/Source/Runtime/SlateRHIRenderer/Private/SlateRHIRenderer.h
+++ b/Engine/Source/Runtime/SlateRHIRenderer/Private/SlateRHIRenderer.h
@@ -19,7 +19,7 @@ const uint32 NumDrawBuffers = 3;
 class FSlateBackBuffer : public FRenderTarget
 {
 public:
-	FSlateBackBuffer(FTexture2DRHIRef& InRenderTargetTexture, FIntPoint InSizeXY)
+	FSlateBackBuffer(FTexture2DRHIRef InRenderTargetTexture, FIntPoint InSizeXY)
 		: SizeXY(InSizeXY)
 	{
 		RenderTargetTextureRHI = InRenderTargetTexture;
@@ -135,14 +135,18 @@ private:
 		void RecreateDepthBuffer_RenderThread();
 	};
 public:
-	/** 
-	 * Constructor. Initializes all rendering resources
-	 * 
-	 * @param InStyle	A style used when rendering elements.  The renderer loads textures from information in this style
-	 */
-	FSlateRHIRenderer();
+	FSlateRHIRenderer( TSharedPtr<FSlateRHIResourceManager> InResourceManager, TSharedPtr<FSlateFontCache> InFontCache, TSharedPtr<FSlateFontMeasure> InFontMeasure );
 	~FSlateRHIRenderer();
 
+	/**
+	 * Creates a projection matrix for use when rendering an SWindow
+	 *
+	 * @param Width 	The width of the window
+	 * @param Height	The height of the window
+	 * @return The created projection matrix
+	 */
+	static FMatrix CreateProjectionMatrix( uint32 Width, uint32 Height );
+
 	/** FSlateRenderer interface */
 	virtual void Initialize() override;
 	virtual void Destroy() override;
diff --git a/Engine/Source/Runtime/SlateRHIRenderer/Private/SlateRHIRendererModule.cpp b/Engine/Source/Runtime/SlateRHIRenderer/Private/SlateRHIRendererModule.cpp
--- a/Engine/Source/Runtime/SlateRHIRenderer/Private/SlateRHIRendererModule.cpp
+++ b/Engine/Source/Runtime/SlateRHIRenderer/Private/SlateRHIRendererModule.cpp
@@ -2,6 +2,33 @@
 
 #include "SlateRHIRendererPrivatePCH.h"
 
+class FSlateRHIFontAtlasFactory : public ISlateFontAtlasFactory
+{
+public:
+	FSlateRHIFontAtlasFactory()
+	{
+		/** Size of each font texture, width and height */
+		AtlasSize = 1024;
+		if (!GIsEditor && GConfig)
+		{
+			GConfig->GetInt(TEXT("SlateRenderer"), TEXT("FontAtlasSize"), AtlasSize, GEngineIni);
+			AtlasSize = FMath::Clamp(AtlasSize, 0, 2048);
+		}
+	}
+
+	virtual ~FSlateRHIFontAtlasFactory()
+	{
+	}
+
+
+	virtual TSharedRef<FSlateFontAtlas> CreateFontAtlas() const override
+	{
+		return MakeShareable(new FSlateFontAtlasRHI(AtlasSize, AtlasSize));
+	}
+private:
+	int32 AtlasSize;
+};
+
 
 /**
  * Implements the Slate RHI Renderer module.
@@ -12,14 +39,46 @@ class FSlateRHIRendererModule
 public:
 
 	// ISlateRHIRendererModule interface
-
 	virtual TSharedRef<FSlateRenderer> CreateSlateRHIRenderer( ) override
 	{
-		return TSharedRef<FSlateRHIRenderer>( new FSlateRHIRenderer );
+		ConditionalCreateResources();
+
+		return MakeShareable( new FSlateRHIRenderer( ResourceManager, FontCache, FontMeasure ) );
 	}
 
+
 	virtual void StartupModule( ) override { }
 	virtual void ShutdownModule( ) override { }
+
+private:
+	/** Creates resource managers if they do not exist */
+	void ConditionalCreateResources()
+	{
+		if( !ResourceManager.IsValid() )
+		{
+			ResourceManager = MakeShareable( new FSlateRHIResourceManager );
+		}
+
+		if( !FontCache.IsValid() )
+		{
+			FontCache = MakeShareable(new FSlateFontCache(MakeShareable(new FSlateRHIFontAtlasFactory)));
+		}
+
+		if( !FontMeasure.IsValid() )
+		{
+			FontMeasure = FSlateFontMeasure::Create(FontCache.ToSharedRef());
+		}
+
+	}
+private:
+	/** Resource manager used for all renderers */
+	TSharedPtr<FSlateRHIResourceManager> ResourceManager;
+
+	/** Font cache used for all renderers */
+	TSharedPtr<FSlateFontCache> FontCache;
+
+	/** Font measure interface used for all renderers */
+	TSharedPtr<FSlateFontMeasure> FontMeasure;
 };
 
 
diff --git a/Engine/Source/Runtime/SlateRHIRenderer/Private/SlateRHIRenderingPolicy.cpp b/Engine/Source/Runtime/SlateRHIRenderer/Private/SlateRHIRenderingPolicy.cpp
--- a/Engine/Source/Runtime/SlateRHIRenderer/Private/SlateRHIRenderingPolicy.cpp
+++ b/Engine/Source/Runtime/SlateRHIRenderer/Private/SlateRHIRenderingPolicy.cpp
@@ -6,23 +6,10 @@
 #include "SlateShaders.h"
 #include "SlateMaterialShader.h"
 #include "SlateRHIResourceManager.h"
-#include "TileRendering.h"
 #include "PreviewScene.h"
 #include "EngineModule.h"
-
-static EPrimitiveType GetRHIPrimitiveType( ESlateDrawPrimitive::Type SlateType )
-{
-	switch( SlateType )
-	{
-	case ESlateDrawPrimitive::LineList:
-		return PT_LineList;
-	case ESlateDrawPrimitive::TriangleList:
-	default:
-		return PT_TriangleList;
-	}
-
-};
-
+#include "SlateUTextureResource.h"
+#include "SlateMaterialResource.h"
 
 FSlateElementIndexBuffer::FSlateElementIndexBuffer()
 	: BufferSize(0)	 
@@ -264,7 +251,7 @@ static FSceneView& CreateSceneView( FSceneViewFamilyContext& ViewFamilyContext,
 	return *View;
 }
 
-void FSlateRHIRenderingPolicy::DrawElements(FRHICommandListImmediate& RHICmdList, const FIntPoint& InViewportSize, FSlateBackBuffer& BackBuffer, const FMatrix& ViewProjectionMatrix, const TArray<FSlateRenderBatch>& RenderBatches)
+void FSlateRHIRenderingPolicy::DrawElements(FRHICommandListImmediate& RHICmdList, FSlateBackBuffer& BackBuffer, const FMatrix& ViewProjectionMatrix, const TArray<FSlateRenderBatch>& RenderBatches)
 {
 	SCOPE_CYCLE_COUNTER( STAT_SlateDrawTime );
 
@@ -321,7 +308,7 @@ void FSlateRHIRenderingPolicy::DrawElements(FRHICommandListImmediate& RHICmdList
 
 		if( !RenderBatch.CustomDrawer.IsValid() )
 		{
-			if( !ShaderResource || ShaderResource->GetType() == ESlateShaderResource::Texture ) 
+			if( !ShaderResource || ShaderResource->GetType() != ESlateShaderResource::Material ) 
 			{
 				FSlateElementPS* PixelShader = GetTexturePixelShader(ShaderType, DrawEffects);
 
@@ -341,7 +328,7 @@ void FSlateRHIRenderingPolicy::DrawElements(FRHICommandListImmediate& RHICmdList
 				RHICmdList.SetBlendState(
 					(RenderBatch.DrawFlags & ESlateBatchDrawFlag::NoBlending)
 					? TStaticBlendState<>::GetRHI()
-					: TStaticBlendState<CW_RGBA, BO_Add, BF_SourceAlpha, BF_InverseSourceAlpha, BO_Add, BF_Zero, BF_One>::GetRHI()
+					: TStaticBlendState<CW_RGBA, BO_Add, BF_SourceAlpha, BF_InverseSourceAlpha, BO_Add, BF_One, BF_InverseSourceAlpha>::GetRHI()
 					);
 #else
 				RHICmdList.SetBlendState(TStaticBlendState<CW_RGB, BO_Add, BF_One, BF_One, BO_Add, BF_Zero, BF_InverseSourceAlpha>::GetRHI());
@@ -359,33 +346,18 @@ void FSlateRHIRenderingPolicy::DrawElements(FRHICommandListImmediate& RHICmdList
 					RHICmdList.SetRasterizerState(TStaticRasterizerState<FM_Solid, CM_None, true>::GetRHI());
 				}
 
+				FTexture2DRHIRef TextureRHI;
+				if(ShaderResource)
+				{
+					TextureRHI = ShaderResource->GetType() == ESlateShaderResource::NativeTexture ? ((FSlateTexture2DRHIRef*)ShaderResource)->GetTypedResource() :  ((FSlateUTextureResource*)ShaderResource)->GetTypedResource();
+				}
 
 				if (ShaderType == ESlateShader::LineSegment)
 				{
-					// Generally looks better without stencil testing although there is minor overdraw on splines
-#if 0
-					// Enable stencil testing for anti-aliased line segments to reduce artifacts from overlapping segments
-					 FDepthStencilStateRHIRef DSLineSegment =
-						TStaticDepthStencilState < false
-						, CF_Always
-						, true
-						, CF_Greater
-						, SO_Keep
-						, SO_SaturatedIncrement
-						, SO_SaturatedIncrement
-						, true
-						, CF_Greater
-						, SO_Keep
-						, SO_SaturatedIncrement
-						, SO_SaturatedIncrement>::GetRHI();
-
-					RHICmdList.SetDepthStencilState( DSLineSegment, 0x01 );
-#endif
-
 					// The lookup table used for splines should clamp when sampling otherwise the results are wrong
-					PixelShader->SetTexture(RHICmdList, ((FSlateTexture2DRHIRef*)ShaderResource)->GetTypedResource(), BilinearClamp );
+					PixelShader->SetTexture(RHICmdList, TextureRHI, BilinearClamp );
 				}
-				else if( ShaderResource && ShaderResource->GetType() == ESlateShaderResource::Texture && IsValidRef( ((FSlateTexture2DRHIRef*)ShaderResource)->GetTypedResource() ) )
+				else if( ShaderResource && IsValidRef( TextureRHI ) )
 				{
 					FSamplerStateRHIRef SamplerState; 
 
@@ -405,7 +377,7 @@ void FSlateRHIRenderingPolicy::DrawElements(FRHICommandListImmediate& RHICmdList
 					{
 						SamplerState = BilinearClamp;
 					}
-					PixelShader->SetTexture(RHICmdList, ((FSlateTexture2DRHIRef*)ShaderResource)->GetTypedResource(), SamplerState);
+					PixelShader->SetTexture(RHICmdList, TextureRHI, SamplerState);
 				}
 				else
 				{
@@ -439,7 +411,7 @@ void FSlateRHIRenderingPolicy::DrawElements(FRHICommandListImmediate& RHICmdList
 					SceneView = &CreateSceneView(SceneViewContext, BackBuffer, ViewProjectionMatrix);
 				}
 
-				FMaterialRenderProxy* MaterialRenderProxy = ((FSlateMaterial*)ShaderResource)->GetRenderProxy();
+				FMaterialRenderProxy* MaterialRenderProxy = ((FSlateMaterialResource*)ShaderResource)->GetRenderProxy();
 
 				const FMaterial* Material = MaterialRenderProxy->GetMaterial(SceneView->GetFeatureLevel());
 
@@ -611,5 +583,18 @@ FSlateMaterialShaderPS* FSlateRHIRenderingPolicy::GetMaterialPixelShader( const
 	return FoundShader ? (FSlateMaterialShaderPS*)FoundShader->GetShaderChecked() : nullptr;
 }
 
+EPrimitiveType FSlateRHIRenderingPolicy::GetRHIPrimitiveType(ESlateDrawPrimitive::Type SlateType)
+{
+	switch(SlateType)
+	{
+	case ESlateDrawPrimitive::LineList:
+		return PT_LineList;
+	case ESlateDrawPrimitive::TriangleList:
+	default:
+		return PT_TriangleList;
+	}
+
+};
+
 
 
diff --git a/Engine/Source/Runtime/SlateRHIRenderer/Private/SlateRHIRenderingPolicy.h b/Engine/Source/Runtime/SlateRHIRenderer/Private/SlateRHIRenderingPolicy.h
--- a/Engine/Source/Runtime/SlateRHIRenderer/Private/SlateRHIRenderingPolicy.h
+++ b/Engine/Source/Runtime/SlateRHIRenderer/Private/SlateRHIRenderingPolicy.h
@@ -177,7 +177,7 @@ public:
 	~FSlateRHIRenderingPolicy();
 	
 	virtual void UpdateBuffers( const FSlateWindowElementList& WindowElementList ) override;
-	virtual void DrawElements(FRHICommandListImmediate& RHICmdList, const FIntPoint& InViewportSize, class FSlateBackBuffer& BackBuffer, const FMatrix& ViewProjectionMatrix, const TArray<FSlateRenderBatch>& RenderBatches);
+	virtual void DrawElements(FRHICommandListImmediate& RHICmdList, class FSlateBackBuffer& BackBuffer, const FMatrix& ViewProjectionMatrix, const TArray<FSlateRenderBatch>& RenderBatches);
 
 	virtual TSharedRef<FSlateFontCache> GetFontCache() override { return FontCache.ToSharedRef(); }
 	virtual TSharedRef<FSlateShaderResourceManager> GetResourceManager() override { return ResourceManager; }
@@ -198,6 +198,8 @@ private:
 	class FSlateElementPS* GetTexturePixelShader( ESlateShader::Type ShaderType, ESlateDrawEffect::Type DrawEffects );
 	class FSlateMaterialShaderPS* GetMaterialPixelShader( const class FMaterial* Material, ESlateShader::Type ShaderType, ESlateDrawEffect::Type DrawEffects );
 
+	/** @return The RHI primitive type from the Slate primitive type */
+	EPrimitiveType GetRHIPrimitiveType(ESlateDrawPrimitive::Type SlateType);
 private:
 	/** Buffers used for rendering */
 	TSlateElementVertexBuffer<FSlateVertex> VertexBuffers[2];
diff --git a/Engine/Source/Runtime/SlateRHIRenderer/Private/SlateRHIResourceManager.cpp b/Engine/Source/Runtime/SlateRHIRenderer/Private/SlateRHIResourceManager.cpp
--- a/Engine/Source/Runtime/SlateRHIRenderer/Private/SlateRHIResourceManager.cpp
+++ b/Engine/Source/Runtime/SlateRHIRenderer/Private/SlateRHIResourceManager.cpp
@@ -2,127 +2,65 @@
 
 #include "SlateRHIRendererPrivatePCH.h"
 #include "ImageWrapper.h"
+#include "SlateNativeTextureResource.h"
+#include "SlateUTextureResource.h"
+#include "SlateMaterialResource.h"
 
-
-
-TSharedPtr<FDynamicTextureResource> FDynamicTextureResource::NullResource = MakeShareable( new FDynamicTextureResource( NULL ) );
-
-
-FDynamicTextureResource::FDynamicTextureResource(FSlateTexture2DRHIRef* ExistingTexture)
-	: TextureObject( NULL )
-	, Proxy(new FSlateShaderResourceProxy)
-	, RHIRefTexture(ExistingTexture != NULL ? ExistingTexture : new FSlateTexture2DRHIRef(NULL, 0, 0))
+TSharedPtr<FSlateDynamicTextureResource> FDynamicResourceMap::GetDynamicTextureResource( FName ResourceName ) const
 {
-	Proxy->Resource = RHIRefTexture;
+	return NativeTextureMap.FindRef( ResourceName );
 }
 
-FDynamicTextureResource::~FDynamicTextureResource()
+TSharedPtr<FSlateUTextureResource> FDynamicResourceMap::GetUTextureResource( UTexture2D* TextureObject ) const
 {
-	if (Proxy)
+	if(TextureObject)
 	{
-		delete Proxy;
+		return UTextureResourceMap.FindRef(TextureObject);
 	}
 
-	if (RHIRefTexture)
-	{
-		delete RHIRefTexture;
-	}
+	return nullptr;
 }
 
-
-FSlateMaterialResource::FSlateMaterialResource(UMaterialInterface& InMaterial, const FVector2D& InImageSize)
-	: MaterialObject(&InMaterial)
-	, Proxy(new FSlateShaderResourceProxy)
-	, SlateMaterial(new FSlateMaterial(InMaterial.GetRenderProxy(false), InImageSize))
-{
-	Proxy->ActualSize = InImageSize.IntPoint();
-	Proxy->Resource = SlateMaterial;
-}
-
-FSlateMaterialResource::~FSlateMaterialResource()
-{
-	if (Proxy)
-	{
-		delete Proxy;
-	}
-
-	if (SlateMaterial)
-	{
-		delete SlateMaterial;
-	}
-}
-
-TSharedPtr<FDynamicTextureResource> FDynamicResourceMap::GetTextureResource( FName ResourceName, UTexture2D* TextureObject ) const
+TSharedPtr<FSlateMaterialResource> FDynamicResourceMap::GetMaterialResource( UMaterialInterface* Material ) const
 {
-	if( TextureObject )
-	{
-		const TSharedPtr<FDynamicTextureResource>& FoundResource = DynamicResourceObjectMap.FindRef( TextureObject );
-		if ( FoundResource.IsValid() )
-		{
-			return FoundResource;
-		}
-		else
-		{
-			// We might have already created a resource for this texture in the Native map. Look that up here.
-			return DynamicNativeTextureMapLookupByTextureObject.FindRef( TextureObject ).Pin();
-		}
-	}
-	else
-	{
-		return DynamicNativeTextureMap.FindRef( ResourceName );
-	}
+	return MaterialResourceMap.FindRef( Material );
 }
 
 
-TSharedPtr<FSlateMaterialResource> FDynamicResourceMap::GetMaterialResource( UMaterialInterface* Material ) const
+void FDynamicResourceMap::AddDynamicTextureResource( FName ResourceName, TSharedRef<FSlateDynamicTextureResource> InResource )
 {
-	return MaterialResourceMap.FindRef( Material );
+	NativeTextureMap.Add( ResourceName, InResource );
 }
 
-
-void FDynamicResourceMap::AddTextureResource( FName ResourceName, UTexture2D* TextureObject, TSharedRef<FDynamicTextureResource> InResource )
+void FDynamicResourceMap::AddUTextureResource( UTexture2D* TextureObject, TSharedRef<FSlateUTextureResource> InResource)
 {
-	if( TextureObject )
+	if(TextureObject)
 	{
-		check( TextureObject == InResource->TextureObject );
-		DynamicResourceObjectMap.Add( TextureObject, InResource );
+		check(TextureObject == InResource->TextureObject);
+		UTextureResourceMap.Add(TextureObject, InResource);
 	}
-	else
-	{
-		DynamicNativeTextureMap.Add( ResourceName, InResource );
-
-		if ( InResource->TextureObject )
-		{
-			DynamicNativeTextureMapLookupByTextureObject.Add(InResource->TextureObject, InResource);
-		}
-	}
-
 }
 
 void FDynamicResourceMap::AddMaterialResource( UMaterialInterface* Material, TSharedRef<FSlateMaterialResource> InMaterialResource )
 {
-	check( Material == InMaterialResource->MaterialObject );
+	check( Material == InMaterialResource->GetMaterialObject() );
 	MaterialResourceMap.Add( Material, InMaterialResource );
 }
 
-void FDynamicResourceMap::RemoveTextureResource( FName ResourceName, UTexture2D* TextureObject )
+void FDynamicResourceMap::RemoveDynamicTextureResource(FName ResourceName)
 {
-	if( TextureObject )
-	{
-		DynamicResourceObjectMap.Remove( TextureObject );
-	}
-	else
-	{
-		const TSharedPtr<FDynamicTextureResource>& ExistingResource = DynamicNativeTextureMap.FindRef( ResourceName );
-		DynamicNativeTextureMap.Remove( ResourceName );
+	NativeTextureMap.Remove(ResourceName);
+}
 
-		if (ExistingResource.IsValid() && ExistingResource->TextureObject)
-		{
-			DynamicNativeTextureMapLookupByTextureObject.Remove(ExistingResource->TextureObject);
-		}
+void FDynamicResourceMap::RemoveUTextureResource( UTexture2D* TextureObject )
+{
+	if(TextureObject)
+	{
+		UTextureResourceMap.Remove(TextureObject);
 	}
 }
 
+
 void FDynamicResourceMap::RemoveMaterialResource( UMaterialInterface* Material )
 {
 	MaterialResourceMap.Remove( Material );
@@ -130,39 +68,25 @@ void FDynamicResourceMap::RemoveMaterialResource( UMaterialInterface* Material )
 
 void FDynamicResourceMap::Empty()
 {
-	DynamicNativeTextureMap.Empty();
-	DynamicNativeTextureMapLookupByTextureObject.Empty();
-	DynamicResourceObjectMap.Empty();
+	NativeTextureMap.Empty();
+	UTextureResourceMap.Empty();
 	MaterialResourceMap.Empty();
 }
 
 void FDynamicResourceMap::ReleaseResources()
 {
-	for (TMap<FName, TSharedPtr<FDynamicTextureResource> >::TIterator It(DynamicNativeTextureMap); It; ++It)
+	for (TMap<FName, TSharedPtr<FSlateDynamicTextureResource> >::TIterator It(NativeTextureMap); It; ++It)
 	{
 		BeginReleaseResource(It.Value()->RHIRefTexture);
 	}
 
-	for (TMap<UObject*, TSharedPtr<FDynamicTextureResource> >::TIterator It(DynamicResourceObjectMap); It; ++It)
-	{
-		BeginReleaseResource(It.Value()->RHIRefTexture);
-	}
 }
 
 void FDynamicResourceMap::AddReferencedObjects(FReferenceCollector& Collector)
 {
-	for(TMap<FName, TSharedPtr<FDynamicTextureResource> >::TIterator It(DynamicNativeTextureMap); It; ++It)
+	for(TMap<UObject*, TSharedPtr<FSlateUTextureResource> >::TIterator It(UTextureResourceMap); It; ++It)
 	{
-		TSharedPtr<FDynamicTextureResource>& Resource = It.Value();
-		if ( Resource.IsValid() && Resource->TextureObject != nullptr )
-		{
-			Collector.AddReferencedObject(Resource->TextureObject);
-		}
-	}
-
-	for(TMap<UObject*, TSharedPtr<FDynamicTextureResource> >::TIterator It(DynamicResourceObjectMap); It; ++It)
-	{
-		TSharedPtr<FDynamicTextureResource>& Resource = It.Value();
+		TSharedPtr<FSlateUTextureResource>& Resource = It.Value();
 		if (Resource.IsValid() && Resource->TextureObject != nullptr)
 		{
 			Collector.AddReferencedObject(Resource->TextureObject);
@@ -374,6 +298,29 @@ FSlateShaderResourceProxy* FSlateRHIResourceManager::GenerateTextureResource( co
 	return NewProxy;
 }
 
+static void LoadUObjectForBrush( const FSlateBrush& InBrush )
+{
+	// Load the utexture
+	FString Path = InBrush.GetResourceName().ToString();
+
+	if( !Path.IsEmpty() )
+	{
+		FString NewPath = Path.RightChop(FSlateBrush::UTextureIdentifier().Len());
+		UObject* TextureObject = LoadObject<UTexture2D>(NULL, *NewPath, NULL, LOAD_None, NULL);
+		FSlateBrush* Brush = const_cast<FSlateBrush*>(&InBrush);
+
+		// Set the texture object to a default texture to prevent constant loading of missing textures
+		if( !TextureObject )
+		{
+			TextureObject = GEngine->DefaultTexture;
+		}
+
+		Brush->SetResourceObject(TextureObject);
+
+		UE_LOG(LogSlate, Warning, TEXT("The texture:// method of loading UTextures for use in Slate is deprecated.  Please convert %s to a Brush Asset"), *Path);
+	}
+}
+
 FSlateShaderResourceProxy* FSlateRHIResourceManager::GetShaderResource( const FSlateBrush& InBrush )
 {
 	check( IsThreadSafeForSlateRendering() );
@@ -389,40 +336,65 @@ FSlateShaderResourceProxy* FSlateRHIResourceManager::GetShaderResource( const FS
 	}
 	else if( InBrush.IsDynamicallyLoaded() || ( InBrush.HasUObject() ) )
 	{
+		if( InBrush.HasUObject() && InBrush.GetResourceObject() == nullptr )
+		{
+			// Hack for loading via the deprecated path
+			LoadUObjectForBrush( InBrush );
+		}
+
 		Texture = GetDynamicTextureResource( InBrush );
 	}
 
 	return Texture;
 }
 
-TSharedPtr<FDynamicTextureResource> FSlateRHIResourceManager::MakeDynamicTextureResource( FName ResourceName, uint32 Width, uint32 Height, const TArray< uint8 >& Bytes )
+TSharedPtr<FSlateDynamicTextureResource> FSlateRHIResourceManager::MakeDynamicTextureResource( FName ResourceName, uint32 Width, uint32 Height, const TArray< uint8 >& Bytes )
 {
-	// Bail out if we already have this texture loaded
-	TSharedPtr<FDynamicTextureResource> TextureResource = DynamicResourceMap.GetTextureResource( ResourceName, nullptr );
-	if( TextureResource.IsValid() )
+	// Make storage for the image
+	FSlateTextureDataRef TextureStorage = MakeShareable( new FSlateTextureData( Width, Height, GPixelFormats[PF_B8G8R8A8].BlockBytes, Bytes ) );
+
+	TSharedPtr<FSlateDynamicTextureResource> TextureResource;
+	// Get a resource from the free list if possible
+	if(DynamicTextureFreeList.Num() > 0)
 	{
-		return TextureResource;
+		TextureResource = DynamicTextureFreeList.Pop();
+	}
+	else
+	{
+		// Free list is empty, we have to allocate a new resource
+		TextureResource = MakeShareable(new FSlateDynamicTextureResource(nullptr));
 	}
 
-	// Make storage for the image
-	FSlateTextureDataRef TextureStorage = MakeShareable( new FSlateTextureData( Width, Height, GPixelFormats[PF_B8G8R8A8].BlockBytes, Bytes ) );
+	TextureResource->Proxy->ActualSize = FIntPoint(TextureStorage->GetWidth(), TextureStorage->GetHeight());
 
-	// Initialize a texture resource
-	TextureResource = InitializeDynamicTextureResource( TextureStorage, NULL );
+
+	// Init render thread data
+	ENQUEUE_UNIQUE_RENDER_COMMAND_TWOPARAMETER(InitNewSlateDynamicTextureResource,
+		FSlateDynamicTextureResource*, TextureResource, TextureResource.Get(),
+		FSlateTextureDataPtr, InNewTextureData, TextureStorage,
+	{
+		if(InNewTextureData.IsValid())
+		{
+			// Set the texture to use as the texture we just loaded
+			TextureResource->RHIRefTexture->SetTextureData(InNewTextureData, PF_B8G8R8A8, TexCreate_SRGB);
+
+		}
+
+		// Initialize and link the rendering resource
+		TextureResource->RHIRefTexture->InitResource();
+	});
 
 	// Map the new resource so we don't have to load again
-	DynamicResourceMap.AddTextureResource( ResourceName, nullptr, TextureResource.ToSharedRef() );
+	DynamicResourceMap.AddDynamicTextureResource( ResourceName, TextureResource.ToSharedRef() );
+	INC_DWORD_STAT_BY(STAT_SlateNumDynamicTextures, 1);
 
 	return TextureResource;
 }
 
-TSharedPtr<FDynamicTextureResource> FSlateRHIResourceManager::MakeDynamicTextureResource(bool bHasUTexture, bool bIsDynamicallyLoaded, FString ResourcePath, FName ResourceName, UTexture2D* InTextureObject)
+TSharedPtr<FSlateUTextureResource> FSlateRHIResourceManager::MakeDynamicUTextureResource(UTexture2D* InTextureObject)
 {
 	// Generated texture resource
-	TSharedPtr<FDynamicTextureResource> TextureResource;
-	
-	// Texture object if any
-	UTexture2D* TextureObject = NULL;
+	TSharedPtr<FSlateUTextureResource> TextureResource;
 
 	// Data for a loaded disk image
 	FNewTextureInfo Info;
@@ -430,128 +402,36 @@ TSharedPtr<FDynamicTextureResource> FSlateRHIResourceManager::MakeDynamicTexture
 	bool bUsingDeprecatedUTexturePath = false;
 
 	bool bSucceeded = false;
-	if( bHasUTexture || InTextureObject != NULL )
+	if( InTextureObject != NULL )
 	{
-		TextureResource = DynamicResourceMap.GetTextureResource( ResourceName, InTextureObject );
+		TextureResource = DynamicResourceMap.GetUTextureResource( InTextureObject );
 		if( TextureResource.IsValid() )
 		{
 			// Bail out of the resource is already loaded
 			return TextureResource;
 		}
 
-		if( InTextureObject )
-		{
-			TextureObject = InTextureObject;
-		}
-		else
-		{
-			// Deprecated path
-			bUsingDeprecatedUTexturePath = true;
-			// Load the utexture
-			FString Path = ResourceName.ToString();
-			Path = Path.RightChop( FSlateBrush::UTextureIdentifier().Len() );
-			TextureObject = LoadObject<UTexture2D>( NULL, *Path, NULL, LOAD_None, NULL );
-		}
-
-		bSucceeded = TextureObject != NULL;
-	}
-	else if( bIsDynamicallyLoaded )
-	{
-		TextureResource = DynamicResourceMap.GetTextureResource(ResourceName, nullptr );
-		if ( TextureResource.IsValid() )
-		{
-			// Bail out of the resource is already loaded
-			return TextureResource;
-		}
-
-		uint32 Width = 0;
-		uint32 Height = 0;
-		TArray<uint8> RawData;
-
-		// Load the image from disk
-		bSucceeded = LoadTexture( ResourceName, ResourcePath, Width, Height, RawData );
-
-		Info.TextureData = MakeShareable( new FSlateTextureData( Width, Height, GPixelFormats[PF_B8G8R8A8].BlockBytes, RawData ) );
+		bSucceeded = true;
 	}
-
 	
 	if( bSucceeded )
-	{
-		TextureResource = InitializeDynamicTextureResource( Info.TextureData, TextureObject );
-	}
-	else
-	{
-		// Add the null texture so we don't continuously try to load it.
-		TextureResource = FDynamicTextureResource::NullResource;
-	}
-
-	if( !bUsingDeprecatedUTexturePath )
-	{
-		DynamicResourceMap.AddTextureResource(ResourceName, TextureObject, TextureResource.ToSharedRef());
-	}
-	else
-	{
-		DynamicResourceMap.AddTextureResource(ResourceName, nullptr, TextureResource.ToSharedRef());
-	}
-
-
-	return TextureResource;
-}
-
-
-
-TSharedRef<FDynamicTextureResource> FSlateRHIResourceManager::InitializeDynamicTextureResource( const FSlateTextureDataPtr& TextureData, UTexture2D* TextureObject )
-{
-	TSharedPtr<FDynamicTextureResource> TextureResource = NULL;
-
-	// Get a resource from the free list if possible
-	if( DynamicTextureFreeList.Num() > 0 )
-	{
-		TextureResource = DynamicTextureFreeList.Pop();
-	}
-	else
 	{
 		// Free list is empty, we have to allocate a new resource
-		TextureResource = MakeShareable( new FDynamicTextureResource( nullptr ) );
-	}
+		TextureResource = MakeShareable(new FSlateUTextureResource(InTextureObject));
 
-	checkSlow( !AccessedUTextures.Contains( TextureObject ) );
+		checkSlow(!AccessedUTextures.Contains(InTextureObject));
 
-	// Init game thread data;
-	TextureResource->TextureObject = TextureObject;
-
-	if( TextureObject )
-	{
-		TextureResource->Proxy->ActualSize = FIntPoint( TextureObject->GetSizeX(), TextureObject->GetSizeY() );
+		TextureResource->Proxy->ActualSize = FIntPoint(InTextureObject->GetSizeX(), InTextureObject->GetSizeY());
 	}
 	else
 	{
-		TextureResource->Proxy->ActualSize = FIntPoint( TextureData->GetWidth(), TextureData->GetHeight() );
+		// Add the null texture so we don't continuously try to load it.
+		TextureResource = FSlateUTextureResource::NullResource;
 	}
 
-	// Init render thread data
-	ENQUEUE_UNIQUE_RENDER_COMMAND_THREEPARAMETER( InitNewSlateDynamicTextureResource,
-		FDynamicTextureResource*, TextureResource, TextureResource.Get(),
-		FSlateTextureDataPtr, InNewTextureData, TextureData,
-		bool, bHasResourceObject, TextureObject != nullptr,
-	{
-		if( bHasResourceObject )
-		{
-			TextureResource->RHIRefTexture->SetRHIRef( nullptr, 0, 0 );
-
-		}
-		else if( InNewTextureData.IsValid() )
-		{
-			// Set the texture to use as the texture we just loaded
-			TextureResource->RHIRefTexture->SetTextureData( InNewTextureData, PF_B8G8R8A8, TexCreate_SRGB );
-
-		}
-
-		// Initialize and link the rendering resource
-		TextureResource->RHIRefTexture->InitResource();
-	})
+	DynamicResourceMap.AddUTextureResource(InTextureObject, TextureResource.ToSharedRef());
 
-	return TextureResource.ToSharedRef();
+	return TextureResource;
 }
 
 
@@ -562,44 +442,74 @@ FSlateShaderResourceProxy* FSlateRHIResourceManager::GetDynamicTextureResource(
 	const FName ResourceName = InBrush.GetResourceName();
 	if ( ResourceName.IsValid() && ResourceName != NAME_None )
 	{
-		TSharedPtr<FDynamicTextureResource> TextureResource = DynamicResourceMap.GetTextureResource( ResourceName, Cast<UTexture2D>( InBrush.GetResourceObject() ) );
+		TSharedPtr<FSlateUTextureResource> TextureResource ;
 
-		if( !TextureResource.IsValid() )
+		if( InBrush.GetResourceObject() != nullptr )
 		{
-			TextureResource = MakeDynamicTextureResource(InBrush.HasUObject(), InBrush.IsDynamicallyLoaded(), GetResourcePath(InBrush), ResourceName, Cast<UTexture2D>( InBrush.GetResourceObject() ) );
+			TSharedPtr<FSlateUTextureResource> TextureResource = DynamicResourceMap.GetUTextureResource( Cast<UTexture2D>( InBrush.GetResourceObject() ) );
 
-#if STATS
-			if( TextureResource.IsValid() )
+			if(!TextureResource.IsValid())
 			{
-				INC_DWORD_STAT_BY( STAT_SlateNumDynamicTextures, 1 );
+				TextureResource = MakeDynamicUTextureResource( Cast<UTexture2D>(InBrush.GetResourceObject() ) );
+				if(TextureResource.IsValid())
+				{
+					INC_DWORD_STAT_BY(STAT_SlateNumDynamicTextures, 1);
+				}
 			}
-#endif
-		}
 
-		if( TextureResource.IsValid())
-		{
-			UTexture2D* TextureObject = TextureResource->TextureObject;
-			if( TextureObject && !AccessedUTextures.Contains( TextureObject ) && TextureObject->Resource )
+			if(TextureResource.IsValid())
 			{
-				// Set the texture rendering resource that should be used.  The UTexture resource could change at any time so we must do this each frame
-				ENQUEUE_UNIQUE_RENDER_COMMAND_TWOPARAMETER( UpdateSlateUTextureResource,
-					FDynamicTextureResource*, DynamicTextureResource, TextureResource.Get(),
-					FTexture*, InFTexture, TextureResource->TextureObject->Resource,
+				UTexture2D* TextureObject = TextureResource->TextureObject;
+				if(TextureObject && !AccessedUTextures.Contains(TextureObject) && TextureObject->Resource)
 				{
-					DynamicTextureResource->RHIRefTexture->SetRHIRef( InFTexture->TextureRHI->GetTexture2D(), InFTexture->GetSizeX(), InFTexture->GetSizeY() );
-					// Let the streaming manager know we are using this texture now
-					InFTexture->LastRenderTime = FApp::GetCurrentTime();
-				});
+					// Set the texture rendering resource that should be used.  The UTexture resource could change at any time so we must do this each frame
+					ENQUEUE_UNIQUE_RENDER_COMMAND_ONEPARAMETER(UpdateSlateUTextureResource,
+						FSlateUTextureResource*, InUTextureResource, TextureResource.Get(),
+						{
+							FTexture* RenderTexture = InUTextureResource->TextureObject->Resource;
+
+							// Let the streaming manager know we are using this texture now
+							RenderTexture->LastRenderTime = FApp::GetCurrentTime();
 
-				AccessedUTextures.Add( TextureObject );
+							// Refresh FTexture
+							InUTextureResource->UpdateRenderResource( RenderTexture );
+
+						});
+
+				
+					AccessedUTextures.Add(TextureObject);
+				}
+
+				return TextureResource->Proxy;
+			}
+		}
+		else
+		{
+			TSharedPtr<FSlateDynamicTextureResource> TextureResource = DynamicResourceMap.GetDynamicTextureResource( ResourceName );
+
+			if( !TextureResource.IsValid() )
+			{
+				uint32 Width; 
+				uint32 Height;
+				TArray<uint8> RawData;
+
+				// Load the image from disk
+				bool bSucceeded = LoadTexture(ResourceName, ResourceName.ToString(), Width, Height, RawData);
+				if(bSucceeded)
+				{
+					TextureResource = MakeDynamicTextureResource(ResourceName, Width, Height, RawData);
+				}
 			}
 
-			return TextureResource->Proxy;
+			if(TextureResource.IsValid())
+			{
+				return TextureResource->Proxy;
+			}
 		}
 	}
 
 	// dynamic texture was not found or loaded
-	return  NULL;
+	return  nullptr;
 }
 
 FSlateShaderResourceProxy* FSlateRHIResourceManager::GetMaterialResource(const FSlateBrush& InBrush)
@@ -621,11 +531,11 @@ FSlateShaderResourceProxy* FSlateRHIResourceManager::GetMaterialResource(const F
 	else
 	{
 		// Keep the resource up to date
-		MaterialResource->SlateMaterial->SetResource(Material->GetRenderProxy(false, false));
-		MaterialResource->Proxy->ActualSize = InBrush.ImageSize.IntPoint();
+		MaterialResource->UpdateRenderResource(Material->GetRenderProxy(false, false));
+		MaterialResource->SlateProxy->ActualSize = InBrush.ImageSize.IntPoint();
 	}
 
-	return MaterialResource->Proxy;
+	return MaterialResource->SlateProxy;
 }
 
 bool FSlateRHIResourceManager::ContainsTexture( const FName& ResourceName ) const
@@ -644,31 +554,43 @@ void FSlateRHIResourceManager::ReleaseDynamicResource( const FSlateBrush& InBrus
 
 		UObject* ResourceObject = InBrush.GetResourceObject();
 
-		TSharedPtr<FDynamicTextureResource> TextureResource = DynamicResourceMap.GetTextureResource( ResourceName, Cast<UTexture2D>( ResourceObject ) );
-
-		if( TextureResource.IsValid() )
+		if( ResourceObject )
 		{
-			//remove it from the accessed textures
-			AccessedUTextures.Remove( TextureResource->TextureObject );
-
-			// Release the rendering resource, its no longer being used
-			BeginReleaseResource( TextureResource->RHIRefTexture );
+			TSharedPtr<FSlateUTextureResource> TextureResource = DynamicResourceMap.GetUTextureResource(Cast<UTexture2D>(ResourceObject));
 
-			//remove it from the texture map
-			DynamicResourceMap.RemoveTextureResource( ResourceName, TextureResource->TextureObject );
-
-			TextureResource->TextureObject = nullptr;
+			if(TextureResource.IsValid())
+			{
+				//remove it from the accessed textures
+				AccessedUTextures.Remove(TextureResource->TextureObject);
+				DynamicResourceMap.RemoveUTextureResource(TextureResource->TextureObject);
 
-			// Add the resource to the free list so it can be reused
-			DynamicTextureFreeList.Add( TextureResource );
-			
-			DEC_DWORD_STAT_BY( STAT_SlateNumDynamicTextures, 1 );
+				DEC_DWORD_STAT_BY(STAT_SlateNumDynamicTextures, 1);
+			}
+			else
+			{
+				UMaterialInterface* Material = Cast<UMaterialInterface>(ResourceObject);
+				DynamicResourceMap.RemoveMaterialResource(Material);
+			}
+		
 		}
 		else
 		{
-			UMaterialInterface* Material = Cast<UMaterialInterface>( ResourceObject );
-			DynamicResourceMap.RemoveMaterialResource( Material );
+			TSharedPtr<FSlateDynamicTextureResource> TextureResource = DynamicResourceMap.GetDynamicTextureResource(ResourceName);
+
+			if( TextureResource.IsValid() )
+			{
+				// Release the rendering resource, its no longer being used
+				BeginReleaseResource(TextureResource->RHIRefTexture);
+
+				//remove it from the texture map
+				DynamicResourceMap.RemoveDynamicTextureResource(ResourceName);
+
+				DynamicTextureFreeList.Add( TextureResource );
+
+				DEC_DWORD_STAT_BY(STAT_SlateNumDynamicTextures, 1);
+			}
 		}
+		
 	}
 }
 
@@ -717,11 +639,6 @@ void FSlateRHIResourceManager::ReleaseResources()
 
 	DynamicResourceMap.ReleaseResources();
 
-	for( int32 ResourceIndex = 0; ResourceIndex < DynamicTextureFreeList.Num(); ++ResourceIndex )
-	{
-		BeginReleaseResource( DynamicTextureFreeList[ResourceIndex]->RHIRefTexture );
-	}
-
 	// Note the base class has texture proxies only which do not need to be released
 }
 
@@ -742,9 +659,9 @@ void FSlateRHIResourceManager::DeleteResources()
 
 	AccessedUTextures.Empty();
 	DynamicResourceMap.Empty();
-	DynamicTextureFreeList.Empty();
 	TextureAtlases.Empty();
 	NonAtlasedTextures.Empty();
+	DynamicTextureFreeList.Empty();
 
 	// Clean up mapping to texture
 	ClearTextureMap();
diff --git a/Engine/Source/Runtime/SlateRHIRenderer/Private/SlateRHIResourceManager.h b/Engine/Source/Runtime/SlateRHIRenderer/Private/SlateRHIResourceManager.h
--- a/Engine/Source/Runtime/SlateRHIRenderer/Private/SlateRHIResourceManager.h
+++ b/Engine/Source/Runtime/SlateRHIRenderer/Private/SlateRHIResourceManager.h
@@ -1,49 +1,27 @@
 // Copyright 1998-2014 Epic Games, Inc. All Rights Reserved.
 
-
 #pragma once
 
-/** Represents a dynamic texture resource for rendering */
-struct FDynamicTextureResource
-{
-	static TSharedPtr<FDynamicTextureResource> NullResource;
-
-	FDynamicTextureResource(FSlateTexture2DRHIRef* ExistingTexture);
-	~FDynamicTextureResource();
-
-	// @todo Slate This is needed for the deprecated utexture loading path
-	UTexture2D* TextureObject;
-
-	FSlateShaderResourceProxy* Proxy;
-	FSlateTexture2DRHIRef* RHIRefTexture;
-};
-
-struct FSlateMaterialResource
-{
-	/** UObject material being used */
-	class UMaterialInterface* MaterialObject;
-	/** Slate proxy used for batching the material */
-	FSlateShaderResourceProxy* Proxy;
-	/** Slate material shader resource.  Resource used to render the object*/
-	FSlateMaterial* SlateMaterial;
-
-	FSlateMaterialResource(UMaterialInterface& InMaterial, const FVector2D& InImageSize);
-	~FSlateMaterialResource();
-};
+class FSlateDynamicTextureResource;
+class FSlateUTextureResource;
+class FSlateMaterialResource;
 
 struct FDynamicResourceMap : public FGCObject
 {
 public:
-	TSharedPtr<FDynamicTextureResource> GetTextureResource( FName ResourceName, UTexture2D* TextureObject ) const;
+	TSharedPtr<FSlateDynamicTextureResource> GetDynamicTextureResource( FName ResourceName ) const;
 
-	TSharedPtr<FSlateMaterialResource> GetMaterialResource( UMaterialInterface* Material ) const;
+	TSharedPtr<FSlateUTextureResource> GetUTextureResource( UTexture2D* TextureObject ) const;
 
-	void AddTextureResource( FName ResourceName, UTexture2D* TextureObject, TSharedRef<FDynamicTextureResource> InResource );
+	TSharedPtr<FSlateMaterialResource> GetMaterialResource( UMaterialInterface* Material ) const;
 
-	void RemoveTextureResource( FName ResourceName, UTexture2D* TextureObject );
+	void AddUTextureResource( UTexture2D* TextureObject, TSharedRef<FSlateUTextureResource> InResource );
+	void RemoveUTextureResource( UTexture2D* TextureObject );
 
-	void AddMaterialResource( UMaterialInterface* Material, TSharedRef<FSlateMaterialResource> InMaterialResource );
+	void AddDynamicTextureResource( FName ResourceName, TSharedRef<FSlateDynamicTextureResource> InResource);
+	void RemoveDynamicTextureResource( FName ResourceName );
 
+	void AddMaterialResource( UMaterialInterface* Material, TSharedRef<FSlateMaterialResource> InResource );
 	void RemoveMaterialResource( UMaterialInterface* Material );
 
 	void Empty();
@@ -54,11 +32,9 @@ public:
 	virtual void AddReferencedObjects(FReferenceCollector& Collector) override;
 
 private:
-	TMap<FName, TSharedPtr<FDynamicTextureResource> > DynamicNativeTextureMap;
+	TMap<FName, TSharedPtr<FSlateDynamicTextureResource> > NativeTextureMap;
 	
-	TMap<UObject*, TWeakPtr<FDynamicTextureResource> > DynamicNativeTextureMapLookupByTextureObject;
-
-	TMap<UObject*, TSharedPtr<FDynamicTextureResource> > DynamicResourceObjectMap;
+	TMap<UObject*, TSharedPtr<FSlateUTextureResource> > UTextureResourceMap;
 
 	/** Map of all material resources */
 	TMap<TWeakObjectPtr<UMaterialInterface>, TSharedPtr<FSlateMaterialResource> > MaterialResourceMap;
@@ -100,14 +76,9 @@ public:
 	/**
 	 * Makes a dynamic texture resource and begins use of it
 	 *
-	 * @param bHasUTexture			Whether we have a UTexture
-	 * @param bIsDynamicallyLoaded	True if resource is to be dynamically loaded
-	 * @param ResourcePath			The resource's image path
-	 * @param ResourceName			The name identifier of the resource
-	 * @param InTextureObject		The utexture object, if applicable
-	 * @return						The dynamic texture resource created
+	 * @param InTextureObject	The texture object to create the resource from
 	 */
-	TSharedPtr<FDynamicTextureResource> MakeDynamicTextureResource(bool bHasUTexture, bool bIsDynamicallyLoaded, FString ResourcePath, FName ResourceName, UTexture2D* InTextureObject);
+	TSharedPtr<FSlateUTextureResource> MakeDynamicUTextureResource( UTexture2D* InTextureObject);
 	
 	/**
 	 * Makes a dynamic texture resource and begins use of it
@@ -117,7 +88,7 @@ public:
 	 * @param Height				The height of the resource
 	 * @param Bytes					The payload containing the resource
 	 */
-	TSharedPtr<FDynamicTextureResource> MakeDynamicTextureResource( FName ResourceName, uint32 Width, uint32 Height, const TArray< uint8 >& Bytes );
+	TSharedPtr<FSlateDynamicTextureResource> MakeDynamicTextureResource( FName ResourceName, uint32 Width, uint32 Height, const TArray< uint8 >& Bytes );
 
 	/**
 	 * Returns true if a texture resource with the passed in resource name is available 
@@ -127,6 +98,15 @@ public:
 	/** Releases a specific dynamic resource */
 	void ReleaseDynamicResource( const FSlateBrush& InBrush );
 
+	/** 
+	 * Creates a new texture from the given texture name
+	 *
+	 * @param TextureName	The name of the texture to load
+	 */
+	virtual bool LoadTexture( const FName& TextureName, const FString& ResourcePath, uint32& Width, uint32& Height, TArray<uint8>& DecodedImage );
+	virtual bool LoadTexture( const FSlateBrush& InBrush, uint32& Width, uint32& Height, TArray<uint8>& DecodedImage );
+
+
 	/**
 	 * Releases rendering resources
 	 */
@@ -161,19 +141,6 @@ private:
 	 * @param Resources	The brush resources to load images for
 	 */
 	void CreateTextures( const TArray< const FSlateBrush* >& Resources );
-	
-	/** 
-	 * Creates a new texture from the given texture name
-	 *
-	 * @param TextureName	The name of the texture to load
-	 */
-	virtual bool LoadTexture( const FName& TextureName, const FString& ResourcePath, uint32& Width, uint32& Height, TArray<uint8>& DecodedImage );
-	virtual bool LoadTexture( const FSlateBrush& InBrush, uint32& Width, uint32& Height, TArray<uint8>& DecodedImage );
-
-	/** 
-	 * Prepares a dynamic texture resource for use based on the provided inputs
-	 */
-	TSharedRef<FDynamicTextureResource> InitializeDynamicTextureResource( const FSlateTextureDataPtr& TextureStorage, UTexture2D* TextureObject );
 
 	/**
 	 * Generates rendering resources for a texture
@@ -203,7 +170,7 @@ private:
 	/** Set of dynamic textures that are currently being accessed */
 	TSet< TWeakObjectPtr<UTexture2D> > AccessedUTextures;
 	/** List of old dynamic resources that are free to use as new resources */
-	TArray< TSharedPtr<FDynamicTextureResource> > DynamicTextureFreeList;
+	TArray< TSharedPtr<FSlateDynamicTextureResource> > DynamicTextureFreeList;
 	/** Static Texture atlases which have been created */
 	TArray<class FSlateTextureAtlasRHI*> TextureAtlases;
 	/** Static Textures created that are not atlased */	
diff --git a/Engine/Source/Runtime/SlateRHIRenderer/Private/SlateUTextureResource.cpp b/Engine/Source/Runtime/SlateRHIRenderer/Private/SlateUTextureResource.cpp
new file mode 100644
--- /dev/null
+++ b/Engine/Source/Runtime/SlateRHIRenderer/Private/SlateUTextureResource.cpp
@@ -0,0 +1,45 @@
+// Copyright 1998-2014 Epic Games, Inc. All Rights Reserved.
+
+#include "SlateRHIRendererPrivatePCH.h"
+#include "SlateUTextureResource.h"
+
+TSharedPtr<FSlateUTextureResource> FSlateUTextureResource::NullResource = MakeShareable( new FSlateUTextureResource(nullptr) );
+
+FSlateUTextureResource::FSlateUTextureResource(UTexture2D* InTexture)
+	: TextureObject(InTexture)
+	, Proxy(new FSlateShaderResourceProxy)
+{
+	if(TextureObject)
+	{
+		Proxy->ActualSize = FIntPoint(TextureObject->GetSizeX(), TextureObject->GetSizeY());
+		Proxy->Resource = this;
+	}
+}
+
+FSlateUTextureResource::~FSlateUTextureResource()
+{
+	if(Proxy)
+	{
+		delete Proxy;
+	}
+}
+
+void FSlateUTextureResource::UpdateRenderResource(FTexture* InFTexture)
+{
+	ShaderResource = InFTexture ? InFTexture->TextureRHI->GetTexture2D() : FTexture2DRHIRef();
+}
+
+uint32 FSlateUTextureResource::GetWidth() const
+{ 
+	return TextureObject->GetSizeX();
+}
+
+uint32 FSlateUTextureResource::GetHeight() const 
+{ 
+	return TextureObject->GetSizeY(); 
+}
+
+ESlateShaderResource::Type FSlateUTextureResource::GetType() const
+{ 
+	return ESlateShaderResource::TextureObject;
+}
diff --git a/Engine/Source/Runtime/SlateRHIRenderer/Private/SlateUTextureResource.h b/Engine/Source/Runtime/SlateRHIRenderer/Private/SlateUTextureResource.h
new file mode 100644
--- /dev/null
+++ b/Engine/Source/Runtime/SlateRHIRenderer/Private/SlateUTextureResource.h
@@ -0,0 +1,32 @@
+// Copyright 1998-2014 Epic Games, Inc. All Rights Reserved.
+
+#pragma once
+
+/** A resource for rendering a UTexture object in Slate */
+class FSlateUTextureResource : public TSlateTexture<FTexture2DRHIRef>
+{
+public:
+	static TSharedPtr<FSlateUTextureResource> FSlateUTextureResource::NullResource;
+
+	FSlateUTextureResource(UTexture2D* InTexture);
+	~FSlateUTextureResource();
+
+	/**
+	 * Updates the renderng resource with a new UTexture resource
+	 */
+	void UpdateRenderResource(FTexture* InFTexture);
+
+	/**
+	 * FSlateShaderRsourceInterface
+	 */
+	virtual uint32 GetWidth() const;
+	virtual uint32 GetHeight() const;
+	virtual ESlateShaderResource::Type GetType() const;
+	
+public:
+	/** Slate rendering proxy */
+	FSlateShaderResourceProxy* Proxy;
+
+	/** Texture UObject.  Note: lifetime is managed externally */
+	UTexture2D* TextureObject;
+};
