Fix CIS.

diff --git a/Engine/Source/Runtime/SlateRHIRenderer/Private/SlateUTextureResource.h b/Engine/Source/Runtime/SlateRHIRenderer/Private/SlateUTextureResource.h
--- a/Engine/Source/Runtime/SlateRHIRenderer/Private/SlateUTextureResource.h
+++ b/Engine/Source/Runtime/SlateRHIRenderer/Private/SlateUTextureResource.h
@@ -6,7 +6,7 @@
 class FSlateUTextureResource : public TSlateTexture<FTexture2DRHIRef>
 {
 public:
-	static TSharedPtr<FSlateUTextureResource> FSlateUTextureResource::NullResource;
+	static TSharedPtr<FSlateUTextureResource> NullResource;
 
 	FSlateUTextureResource(UTexture2D* InTexture);
 	~FSlateUTextureResource();
