Fix CIS (ambigious ternary operator).

diff --git a/Engine/Source/Runtime/SlateRHIRenderer/Private/SlateUTextureResource.cpp b/Engine/Source/Runtime/SlateRHIRenderer/Private/SlateUTextureResource.cpp
--- a/Engine/Source/Runtime/SlateRHIRenderer/Private/SlateUTextureResource.cpp
+++ b/Engine/Source/Runtime/SlateRHIRenderer/Private/SlateUTextureResource.cpp
@@ -26,7 +26,7 @@ FSlateUTextureResource::~FSlateUTextureResource()
 
 void FSlateUTextureResource::UpdateRenderResource(FTexture* InFTexture)
 {
-	ShaderResource = InFTexture ? InFTexture->TextureRHI->GetTexture2D() : FTexture2DRHIRef();
+	ShaderResource = InFTexture ? FTexture2DRHIRef(InFTexture->TextureRHI->GetTexture2D()) : FTexture2DRHIRef();
 }
 
 uint32 FSlateUTextureResource::GetWidth() const
