Fix Mac CIS (initialization order).

diff --git a/Engine/Source/Runtime/SlateRHIRenderer/Private/SlateUTextureResource.cpp b/Engine/Source/Runtime/SlateRHIRenderer/Private/SlateUTextureResource.cpp
--- a/Engine/Source/Runtime/SlateRHIRenderer/Private/SlateUTextureResource.cpp
+++ b/Engine/Source/Runtime/SlateRHIRenderer/Private/SlateUTextureResource.cpp
@@ -6,8 +6,8 @@
 TSharedPtr<FSlateUTextureResource> FSlateUTextureResource::NullResource = MakeShareable( new FSlateUTextureResource(nullptr) );
 
 FSlateUTextureResource::FSlateUTextureResource(UTexture2D* InTexture)
-	: TextureObject(InTexture)
-	, Proxy(new FSlateShaderResourceProxy)
+	: Proxy(new FSlateShaderResourceProxy)
+	, TextureObject(InTexture)
 {
 	if(TextureObject)
 	{
