Fixed dynamic Slate UTexture resource management race condition.

[Fixed by MattK @ BillyB's machine.]

diff --git a/Engine/Source/Runtime/SlateRHIRenderer/Private/SlateRHIResourceManager.cpp b/Engine/Source/Runtime/SlateRHIRenderer/Private/SlateRHIResourceManager.cpp
--- a/Engine/Source/Runtime/SlateRHIRenderer/Private/SlateRHIResourceManager.cpp
+++ b/Engine/Source/Runtime/SlateRHIRenderer/Private/SlateRHIResourceManager.cpp
@@ -416,12 +416,23 @@ TSharedPtr<FSlateUTextureResource> FSlateRHIResourceManager::MakeDynamicUTexture
 	
 	if( bSucceeded )
 	{
-		// Free list is empty, we have to allocate a new resource
-		TextureResource = MakeShareable(new FSlateUTextureResource(InTextureObject));
 
-		checkSlow(!AccessedUTextures.Contains(InTextureObject));
+		// Get a resource from the free list if possible
+		if (UTextureFreeList.Num() > 0)
+		{
+			TextureResource = UTextureFreeList.Pop(); 
+			TextureResource->TextureObject = InTextureObject;
+		}
+		else
+		{
+			// Free list is empty, we have to allocate a new resource
+			TextureResource = MakeShareable(new FSlateUTextureResource(InTextureObject));
+		
+		}
 
 		TextureResource->Proxy->ActualSize = FIntPoint(InTextureObject->GetSizeX(), InTextureObject->GetSizeY());
+
+		checkSlow(!AccessedUTextures.Contains(InTextureObject));
 	}
 	else
 	{
@@ -446,11 +457,13 @@ FSlateShaderResourceProxy* FSlateRHIResourceManager::GetDynamicTextureResource(
 
 		if( InBrush.GetResourceObject() != nullptr )
 		{
-			TSharedPtr<FSlateUTextureResource> TextureResource = DynamicResourceMap.GetUTextureResource( Cast<UTexture2D>( InBrush.GetResourceObject() ) );
+			UTexture2D* TextureObject = CastChecked<UTexture2D>(InBrush.GetResourceObject());
+
+			TSharedPtr<FSlateUTextureResource> TextureResource = DynamicResourceMap.GetUTextureResource(TextureObject);
 
 			if(!TextureResource.IsValid())
 			{
-				TextureResource = MakeDynamicUTextureResource( Cast<UTexture2D>(InBrush.GetResourceObject() ) );
+				TextureResource = MakeDynamicUTextureResource(TextureObject);
 				if(TextureResource.IsValid())
 				{
 					INC_DWORD_STAT_BY(STAT_SlateNumDynamicTextures, 1);
@@ -564,6 +577,8 @@ void FSlateRHIResourceManager::ReleaseDynamicResource( const FSlateBrush& InBrus
 				AccessedUTextures.Remove(TextureResource->TextureObject);
 				DynamicResourceMap.RemoveUTextureResource(TextureResource->TextureObject);
 
+				UTextureFreeList.Add(TextureResource);
+
 				DEC_DWORD_STAT_BY(STAT_SlateNumDynamicTextures, 1);
 			}
 			else
@@ -662,6 +677,7 @@ void FSlateRHIResourceManager::DeleteResources()
 	TextureAtlases.Empty();
 	NonAtlasedTextures.Empty();
 	DynamicTextureFreeList.Empty();
+	UTextureFreeList.Empty();
 
 	// Clean up mapping to texture
 	ClearTextureMap();
diff --git a/Engine/Source/Runtime/SlateRHIRenderer/Private/SlateRHIResourceManager.h b/Engine/Source/Runtime/SlateRHIRenderer/Private/SlateRHIResourceManager.h
--- a/Engine/Source/Runtime/SlateRHIRenderer/Private/SlateRHIResourceManager.h
+++ b/Engine/Source/Runtime/SlateRHIRenderer/Private/SlateRHIResourceManager.h
@@ -169,6 +169,8 @@ private:
 	FDynamicResourceMap DynamicResourceMap;
 	/** Set of dynamic textures that are currently being accessed */
 	TSet< TWeakObjectPtr<UTexture2D> > AccessedUTextures;
+	/** List of old utexture resources that are free to use as new resources */
+	TArray< TSharedPtr<FSlateUTextureResource> > UTextureFreeList;
 	/** List of old dynamic resources that are free to use as new resources */
 	TArray< TSharedPtr<FSlateDynamicTextureResource> > DynamicTextureFreeList;
 	/** Static Texture atlases which have been created */
