Expose ini settings for games to tweak the size that textures can be in a Slate texture atlas

diff --git a/Engine/Source/Runtime/SlateRHIRenderer/Private/SlateRHIResourceManager.cpp b/Engine/Source/Runtime/SlateRHIRenderer/Private/SlateRHIResourceManager.cpp
--- a/Engine/Source/Runtime/SlateRHIRenderer/Private/SlateRHIResourceManager.cpp
+++ b/Engine/Source/Runtime/SlateRHIRenderer/Private/SlateRHIResourceManager.cpp
@@ -97,6 +97,7 @@ void FDynamicResourceMap::AddReferencedObjects(FReferenceCollector& Collector)
 
 FSlateRHIResourceManager::FSlateRHIResourceManager()
 {
+	MaxAltasedTextureSize = FIntPoint(256,256);
 	if( GIsEditor )
 	{
 		AtlasSize = 2048;
@@ -109,6 +110,15 @@ FSlateRHIResourceManager::FSlateRHIResourceManager()
 			int32 RequestedSize = 1024;
 			GConfig->GetInt( TEXT("SlateRenderer"), TEXT("TextureAtlasSize"), RequestedSize, GEngineIni );
 			AtlasSize = FMath::Clamp<uint32>( RequestedSize, 0, 2048 );
+
+			int32 MaxAtlasedTextureWidth = 256;
+			int32 MaxAtlasedTextureHeight = 256;
+			GConfig->GetInt( TEXT("SlateRenderer"), TEXT("MaxAtlasedTextureWidth"), MaxAtlasedTextureWidth, GEngineIni );
+			GConfig->GetInt( TEXT("SlateRenderer"), TEXT("MaxAtlasedTextureHeight"),MaxAtlasedTextureHeight, GEngineIni );
+
+			// Max texture size cannot be larger than the max size of the atlas
+			MaxAltasedTextureSize.X = FMath::Clamp<int32>( MaxAtlasedTextureWidth, 0, AtlasSize );
+			MaxAltasedTextureSize.Y = FMath::Clamp<int32>( MaxAtlasedTextureHeight, 0, AtlasSize );
 		}
 	}
 }
@@ -150,7 +160,7 @@ void FSlateRHIResourceManager::CreateTextures( const TArray< const FSlateBrush*
 
 				Info.TextureData = MakeShareable( new FSlateTextureData( Width, Height, Stride, RawData ) );
 
-				const bool bTooLargeForAtlas = (Width >= 256 || Height >= 256 || Width >= AtlasSize || Height >= AtlasSize );
+				const bool bTooLargeForAtlas = (Width >= (uint32)MaxAltasedTextureSize.X || Height >= (uint32)MaxAltasedTextureSize.Y || Width >= AtlasSize || Height >= AtlasSize );
 
 				Info.bShouldAtlas &= !bTooLargeForAtlas;
 
diff --git a/Engine/Source/Runtime/SlateRHIRenderer/Private/SlateRHIResourceManager.h b/Engine/Source/Runtime/SlateRHIRenderer/Private/SlateRHIResourceManager.h
--- a/Engine/Source/Runtime/SlateRHIRenderer/Private/SlateRHIResourceManager.h
+++ b/Engine/Source/Runtime/SlateRHIRenderer/Private/SlateRHIResourceManager.h
@@ -179,6 +179,8 @@ private:
 	TArray<class FSlateTexture2DRHIRef*> NonAtlasedTextures;
 	/** The size of each texture atlas (square) */
 	uint32 AtlasSize;
+	/** This max size of each texture in an atlas */
+	FIntPoint MaxAltasedTextureSize;
 
 };
 
