Slate RHI.  Block against users creating multiple brushes of the same name.

diff --git a/Engine/Source/Runtime/SlateRHIRenderer/Private/SlateRHIRenderer.cpp b/Engine/Source/Runtime/SlateRHIRenderer/Private/SlateRHIRenderer.cpp
--- a/Engine/Source/Runtime/SlateRHIRenderer/Private/SlateRHIRenderer.cpp
+++ b/Engine/Source/Runtime/SlateRHIRenderer/Private/SlateRHIRenderer.cpp
@@ -1009,13 +1009,15 @@ FIntPoint FSlateRHIRenderer::GenerateDynamicImageResource(const FName InTextureN
 	uint32 Height = 0;
 	TArray<uint8> RawData;
 
-	// Load the image from disk
-	bool bSucceeded = ResourceManager->LoadTexture(InTextureName, InTextureName.ToString(), Width, Height, RawData);
-
-	TSharedPtr<FSlateDynamicTextureResource> TextureResource;
-	if( bSucceeded )
+	TSharedPtr<FSlateDynamicTextureResource> TextureResource = ResourceManager->GetDynamicTextureResourceByName( InTextureName );
+	if( !TextureResource.IsValid() )
 	{
-		TextureResource = ResourceManager->MakeDynamicTextureResource(InTextureName, Width, Height, RawData);
+		// Load the image from disk
+		bool bSucceeded = ResourceManager->LoadTexture(InTextureName, InTextureName.ToString(), Width, Height, RawData);
+		if (bSucceeded)
+		{
+			TextureResource = ResourceManager->MakeDynamicTextureResource(InTextureName, Width, Height, RawData);
+		}
 	}
 
 	return TextureResource.IsValid() ? TextureResource->Proxy->ActualSize : FIntPoint( 0, 0 );
@@ -1025,7 +1027,12 @@ bool FSlateRHIRenderer::GenerateDynamicImageResource( FName ResourceName, uint32
 {
 	check( IsInGameThread() );
 
-	TSharedPtr<FSlateDynamicTextureResource> TextureResource = ResourceManager->MakeDynamicTextureResource( ResourceName, Width, Height, Bytes );
+	
+	TSharedPtr<FSlateDynamicTextureResource> TextureResource = ResourceManager->GetDynamicTextureResourceByName( ResourceName );
+	if( !TextureResource.IsValid() )
+	{
+		TextureResource = ResourceManager->MakeDynamicTextureResource( ResourceName, Width, Height, Bytes );
+	}
 	return TextureResource.IsValid();
 }
 
diff --git a/Engine/Source/Runtime/SlateRHIRenderer/Private/SlateRHIResourceManager.cpp b/Engine/Source/Runtime/SlateRHIRenderer/Private/SlateRHIResourceManager.cpp
--- a/Engine/Source/Runtime/SlateRHIRenderer/Private/SlateRHIResourceManager.cpp
+++ b/Engine/Source/Runtime/SlateRHIRenderer/Private/SlateRHIResourceManager.cpp
@@ -352,7 +352,7 @@ FSlateShaderResourceProxy* FSlateRHIResourceManager::GetShaderResource( const FS
 			LoadUObjectForBrush( InBrush );
 		}
 
-		Texture = GetDynamicTextureResource( InBrush );
+		Texture = FindOrCreateDynamicTextureResource( InBrush );
 	}
 
 	return Texture;
@@ -401,6 +401,11 @@ TSharedPtr<FSlateDynamicTextureResource> FSlateRHIResourceManager::MakeDynamicTe
 	return TextureResource;
 }
 
+TSharedPtr<FSlateDynamicTextureResource> FSlateRHIResourceManager::GetDynamicTextureResourceByName( FName ResourceName )
+{
+	return DynamicResourceMap.GetDynamicTextureResource( ResourceName );
+}
+
 TSharedPtr<FSlateUTextureResource> FSlateRHIResourceManager::MakeDynamicUTextureResource(UTexture2D* InTextureObject)
 {
 	// Generated texture resource
@@ -456,7 +461,7 @@ TSharedPtr<FSlateUTextureResource> FSlateRHIResourceManager::MakeDynamicUTexture
 }
 
 
-FSlateShaderResourceProxy* FSlateRHIResourceManager::GetDynamicTextureResource( const FSlateBrush& InBrush )
+FSlateShaderResourceProxy* FSlateRHIResourceManager::FindOrCreateDynamicTextureResource(const FSlateBrush& InBrush)
 {
 	check( IsThreadSafeForSlateRendering() );
 
diff --git a/Engine/Source/Runtime/SlateRHIRenderer/Private/SlateRHIResourceManager.h b/Engine/Source/Runtime/SlateRHIRenderer/Private/SlateRHIResourceManager.h
--- a/Engine/Source/Runtime/SlateRHIRenderer/Private/SlateRHIResourceManager.h
+++ b/Engine/Source/Runtime/SlateRHIRenderer/Private/SlateRHIResourceManager.h
@@ -90,6 +90,13 @@ public:
 	 */
 	TSharedPtr<FSlateDynamicTextureResource> MakeDynamicTextureResource( FName ResourceName, uint32 Width, uint32 Height, const TArray< uint8 >& Bytes );
 
+	/**
+	 * Find a dynamic texture resource 
+	 *
+	 * @param ResourceName			The name identifier of the resource
+	 */
+	TSharedPtr<FSlateDynamicTextureResource> GetDynamicTextureResourceByName( FName ResourceName );
+
 	/**
 	 * Returns true if a texture resource with the passed in resource name is available 
 	 */
@@ -155,7 +162,7 @@ private:
 	 *
 	 * @param InBrush	Slate brush for the dynamic resource
 	 */
-	FSlateShaderResourceProxy* GetDynamicTextureResource( const FSlateBrush& InBrush );
+	FSlateShaderResourceProxy* FindOrCreateDynamicTextureResource( const FSlateBrush& InBrush );
 
 	/**
 	 * Returns a rendering resource for a material
