Fix shutdown crash on mac

diff --git a/Engine/Source/Runtime/SlateRHIRenderer/Private/SlateRHIResourceManager.cpp b/Engine/Source/Runtime/SlateRHIRenderer/Private/SlateRHIResourceManager.cpp
--- a/Engine/Source/Runtime/SlateRHIRenderer/Private/SlateRHIResourceManager.cpp
+++ b/Engine/Source/Runtime/SlateRHIRenderer/Private/SlateRHIResourceManager.cpp
@@ -79,6 +79,11 @@ void FDynamicResourceMap::ReleaseResources()
 	{
 		BeginReleaseResource(It.Value()->RHIRefTexture);
 	}
+	
+	for (TMap<UObject*, TSharedPtr<FSlateUTextureResource> >::TIterator It(UTextureResourceMap); It; ++It)
+	{
+		It.Value()->UpdateRenderResource(nullptr);
+	}
 
 }
 
