Fix crash using slate materials due to uninitialized uniform buffer parameters.

diff --git a/Engine/Source/Runtime/SlateRHIRenderer/Private/SlateRHIRenderingPolicy.cpp b/Engine/Source/Runtime/SlateRHIRenderer/Private/SlateRHIRenderingPolicy.cpp
--- a/Engine/Source/Runtime/SlateRHIRenderer/Private/SlateRHIRenderingPolicy.cpp
+++ b/Engine/Source/Runtime/SlateRHIRenderer/Private/SlateRHIRenderingPolicy.cpp
@@ -232,6 +232,9 @@ static FSceneView& CreateSceneView( FSceneViewFamilyContext& ViewFamilyContext,
 	ViewUniformShaderParameters.RealTime = View->Family->CurrentRealTime;
 	ViewUniformShaderParameters.Random = FMath::Rand();
 	ViewUniformShaderParameters.FrameNumber = View->FrameNumber;
+	ViewUniformShaderParameters.DirectionalLightShadowTexture = GWhiteTexture->TextureRHI;
+	ViewUniformShaderParameters.DirectionalLightShadowSampler = TStaticSamplerState<SF_Point, AM_Clamp, AM_Clamp, AM_Clamp>::GetRHI();
+
 
 	ViewUniformShaderParameters.ScreenToWorld = FMatrix(
 		FPlane(1, 0, 0, 0),
