Fix material editor preview tiles drawing over everything

diff --git a/Engine/Source/Editor/GraphEditor/Private/MaterialNodes/SGraphNodeMaterialBase.cpp b/Engine/Source/Editor/GraphEditor/Private/MaterialNodes/SGraphNodeMaterialBase.cpp
--- a/Engine/Source/Editor/GraphEditor/Private/MaterialNodes/SGraphNodeMaterialBase.cpp
+++ b/Engine/Source/Editor/GraphEditor/Private/MaterialNodes/SGraphNodeMaterialBase.cpp
@@ -108,7 +108,7 @@ void FPreviewViewport::OnDrawViewport( const FGeometry& AllottedGeometry, const
 	if (PreviewElement->BeginRenderingCanvas( CanvasRect, ClippingRect, MaterialNode->GetExpressionPreview(), bIsRealtime ))
 	{
 		// Draw above everything else
-		uint32 PreviewLayer = MAX_uint32;//LayerId+1;//
+		uint32 PreviewLayer = LayerId+1;
 		FSlateDrawElement::MakeCustom( OutDrawElements, PreviewLayer, PreviewElement );
 	}
 }
diff --git a/Engine/Source/Runtime/SlateRHIRenderer/Private/SlateRHIRenderingPolicy.cpp b/Engine/Source/Runtime/SlateRHIRenderer/Private/SlateRHIRenderingPolicy.cpp
--- a/Engine/Source/Runtime/SlateRHIRenderer/Private/SlateRHIRenderingPolicy.cpp
+++ b/Engine/Source/Runtime/SlateRHIRenderer/Private/SlateRHIRenderingPolicy.cpp
@@ -467,6 +467,11 @@ void FSlateRHIRenderingPolicy::DrawElements(FRHICommandListImmediate& RHICmdList
 			{
 				// This element is custom and has no Slate geometry.  Tell it to render itself now
 				CustomDrawer->DrawRenderThread(RHICmdList, &BackBuffer.GetRenderTargetTexture());
+
+				// Something may have messed with the viewport size so set it back to the full target.
+				RHICmdList.SetViewport( 0,0,0,BackBuffer.GetSizeXY().X, BackBuffer.GetSizeXY().Y, 0.0f ); 
+				RHICmdList.SetStreamSource(0, VertexBuffer.VertexBufferRHI, sizeof(FSlateVertex), 0);
+
 			}
 		}
 
