Don't allow slate to render when running with the null rhi. Fixes a crash reported in TTP 347641.

diff --git a/Engine/Source/Runtime/SlateRHIRenderer/Private/SlateRHIRenderer.cpp b/Engine/Source/Runtime/SlateRHIRenderer/Private/SlateRHIRenderer.cpp
--- a/Engine/Source/Runtime/SlateRHIRenderer/Private/SlateRHIRenderer.cpp
+++ b/Engine/Source/Runtime/SlateRHIRenderer/Private/SlateRHIRenderer.cpp
@@ -623,11 +623,15 @@ void FSlateRHIRenderer::DrawWindows_Private( FSlateDrawBuffer& WindowDrawBuffer
 					// the FSlateWindowElementList structure
 					Params.SlateWindow = Window.Get();
 
-					ENQUEUE_UNIQUE_RENDER_COMMAND_ONEPARAMETER( SlateDrawWindowsCommand, 
-						FSlateDrawWindowCommandParams, Params, Params,
+					// Skip the actual draw if we're in a headless execution environment
+					if (GIsClient && !IsRunningCommandlet() && !GUsingNullRHI)
 					{
-						Params.Renderer->DrawWindow_RenderThread(RHICmdList, *Params.ViewportInfo, *Params.WindowElementList, Params.bLockToVsync);
-					});
+						ENQUEUE_UNIQUE_RENDER_COMMAND_ONEPARAMETER(SlateDrawWindowsCommand,
+							FSlateDrawWindowCommandParams, Params, Params,
+							{
+								Params.Renderer->DrawWindow_RenderThread(RHICmdList, *Params.ViewportInfo, *Params.WindowElementList, Params.bLockToVsync);
+							});
+					}
 
 					SlateWindowRendered.Broadcast( *Params.SlateWindow, &ViewInfo->ViewportRHI );
 
@@ -839,7 +843,7 @@ void FSlateRHIRenderer::CopyWindowsToVirtualScreenBuffer(const TArray<FString>&
 				DrawWindowToBuffer,
 				FDrawWindowToBufferContext,Context,DrawWindowToBufferContext,
 			{
-				const auto FeatureLevel = GRHIFeatureLevel;
+				const auto FeatureLevel = GMaxRHIFeatureLevel;
 				auto ShaderMap = GetGlobalShaderMap(FeatureLevel);
 
 				TShaderMapRef<FScreenVS> VertexShader(ShaderMap);
