Handle WindowedFullscreen properly in FSlateRHIRenderer::RestoreSystemResolution introduced in CL #2254719, which was relying on GSystemResolution which will be wrong when entering windowed-fullscreen (the default  & presently only option on OS X) in a manner other than via the in-game command-line. Windowed-fullscreen always uses the full display resolution, so use the same code from FSlateRHIRenderer::UpdateFullscreenState to ensure that we respect the current state for this mode.

diff --git a/Engine/Source/Runtime/SlateRHIRenderer/Private/SlateRHIRenderer.cpp b/Engine/Source/Runtime/SlateRHIRenderer/Private/SlateRHIRenderer.cpp
--- a/Engine/Source/Runtime/SlateRHIRenderer/Private/SlateRHIRenderer.cpp
+++ b/Engine/Source/Runtime/SlateRHIRenderer/Private/SlateRHIRenderer.cpp
@@ -358,7 +358,19 @@ void FSlateRHIRenderer::RestoreSystemResolution(const TSharedRef<SWindow> InWind
 {
 	if (!GIsEditor)
 	{
-		FSystemResolution::RequestResolutionChange(GSystemResolution.ResX, GSystemResolution.ResY, GSystemResolution.WindowMode);
+		uint32 ResX = GSystemResolution.ResX;
+		uint32 ResY = GSystemResolution.ResY;
+		
+		if( InWindow->GetWindowMode() == EWindowMode::WindowedFullscreen )
+		{
+			FViewportInfo* ViewInfo = WindowToViewportInfo.FindRef( &InWindow.Get() );
+			if( ViewInfo )
+			{
+				ResX = ViewInfo->Width;
+				ResY = ViewInfo->Height;
+			}
+		}
+		FSystemResolution::RequestResolutionChange(ResX, ResY, GSystemResolution.WindowMode);
 		GSystemResolution.ResX = GSystemResolution.ResY = 0;
 	}
 }
