#UE4
- Fixed slate material brush crashing

diff --git a/Engine/Source/Runtime/SlateRHIRenderer/Private/SlateMaterialResource.cpp b/Engine/Source/Runtime/SlateRHIRenderer/Private/SlateMaterialResource.cpp
--- a/Engine/Source/Runtime/SlateRHIRenderer/Private/SlateMaterialResource.cpp
+++ b/Engine/Source/Runtime/SlateRHIRenderer/Private/SlateMaterialResource.cpp
@@ -23,3 +23,19 @@ FSlateMaterialResource::~FSlateMaterialResource()
 	}
 }
 
+void FSlateMaterialResource::UpdateMaterial(UMaterialInterface& InMaterialResource, const FVector2D& InImageSize)
+{
+	MaterialObject = &InMaterialResource;
+	RenderProxy = InMaterialResource.GetRenderProxy(false);
+	if( !SlateProxy )
+	{
+		SlateProxy = new FSlateShaderResourceProxy;
+	}
+
+	SlateProxy->ActualSize = InImageSize.IntPoint();
+	SlateProxy->Resource = this;
+
+	Width = FMath::RoundToInt(InImageSize.X);
+	Height = FMath::RoundToInt(InImageSize.Y);
+}
+
diff --git a/Engine/Source/Runtime/SlateRHIRenderer/Private/SlateMaterialResource.h b/Engine/Source/Runtime/SlateRHIRenderer/Private/SlateMaterialResource.h
--- a/Engine/Source/Runtime/SlateRHIRenderer/Private/SlateMaterialResource.h
+++ b/Engine/Source/Runtime/SlateRHIRenderer/Private/SlateMaterialResource.h
@@ -15,6 +15,8 @@ public:
 	virtual uint32 GetHeight() const { return Height; }
 	virtual ESlateShaderResource::Type GetType() const override { return ESlateShaderResource::Material; }
 
+	void UpdateMaterial(UMaterialInterface& InMaterialResource, const FVector2D& InImageSize);
+
 	/** Updates the rendering resource */
 	void UpdateRenderResource(FMaterialRenderProxy* Proxy) { RenderProxy = Proxy; }
 
diff --git a/Engine/Source/Runtime/SlateRHIRenderer/Private/SlateRHIResourceManager.cpp b/Engine/Source/Runtime/SlateRHIRenderer/Private/SlateRHIResourceManager.cpp
--- a/Engine/Source/Runtime/SlateRHIRenderer/Private/SlateRHIResourceManager.cpp
+++ b/Engine/Source/Runtime/SlateRHIRenderer/Private/SlateRHIResourceManager.cpp
@@ -559,10 +559,19 @@ FSlateShaderResourceProxy* FSlateRHIResourceManager::GetMaterialResource(const F
 	TSharedPtr<FSlateMaterialResource> MaterialResource = DynamicResourceMap.GetMaterialResource(Material);
 	if (!MaterialResource.IsValid())
 	{
-		TSharedRef<FSlateMaterialResource> NewMaterialResource = MakeShareable(new FSlateMaterialResource(*Material, InBrush.ImageSize));
-		DynamicResourceMap.AddMaterialResource( Material, NewMaterialResource );
-
-		MaterialResource = NewMaterialResource;
+		// Get a resource from the free list if possible
+		if(MaterialResourceFreeList.Num() > 0)
+		{
+			MaterialResource = MaterialResourceFreeList.Pop();
+		
+			MaterialResource->UpdateMaterial( *Material, InBrush.ImageSize );
+		}
+		else
+		{
+			MaterialResource = MakeShareable(new FSlateMaterialResource(*Material, InBrush.ImageSize));
+		}
+		
+		DynamicResourceMap.AddMaterialResource(Material, MaterialResource.ToSharedRef());
 	}
 	else
 	{
@@ -607,7 +616,15 @@ void FSlateRHIResourceManager::ReleaseDynamicResource( const FSlateBrush& InBrus
 			else
 			{
 				UMaterialInterface* Material = Cast<UMaterialInterface>(ResourceObject);
+
+				TSharedPtr<FSlateMaterialResource> MaterialResource = DynamicResourceMap.GetMaterialResource(Material);
+				
 				DynamicResourceMap.RemoveMaterialResource(Material);
+
+				if (MaterialResource.IsValid())
+				{
+					MaterialResourceFreeList.Add( MaterialResource );
+				}
 			}
 		
 		}
@@ -700,6 +717,7 @@ void FSlateRHIResourceManager::DeleteResources()
 	TextureAtlases.Empty();
 	NonAtlasedTextures.Empty();
 	DynamicTextureFreeList.Empty();
+	MaterialResourceFreeList.Empty();
 	UTextureFreeList.Empty();
 
 	// Clean up mapping to texture
diff --git a/Engine/Source/Runtime/SlateRHIRenderer/Private/SlateRHIResourceManager.h b/Engine/Source/Runtime/SlateRHIRenderer/Private/SlateRHIResourceManager.h
--- a/Engine/Source/Runtime/SlateRHIRenderer/Private/SlateRHIResourceManager.h
+++ b/Engine/Source/Runtime/SlateRHIRenderer/Private/SlateRHIResourceManager.h
@@ -179,7 +179,9 @@ private:
 	/** List of old utexture resources that are free to use as new resources */
 	TArray< TSharedPtr<FSlateUTextureResource> > UTextureFreeList;
 	/** List of old dynamic resources that are free to use as new resources */
-	TArray< TSharedPtr<FSlateDynamicTextureResource> > DynamicTextureFreeList;
+	TArray< TSharedPtr<FSlateDynamicTextureResource> > DynamicTextureFreeList;\
+	/** List of old material resources that are free to use as new resources */
+	TArray< TSharedPtr<FSlateMaterialResource> > MaterialResourceFreeList;
 	/** Static Texture atlases which have been created */
 	TArray<class FSlateTextureAtlasRHI*> TextureAtlases;
 	/** Static Textures created that are not atlased */	
