Fix shutdown crash if the object system is shut down before Slate object rendering resources are cleaned up.

diff --git a/Engine/Source/Runtime/SlateRHIRenderer/Private/SlateRHIResourceManager.cpp b/Engine/Source/Runtime/SlateRHIRenderer/Private/SlateRHIResourceManager.cpp
--- a/Engine/Source/Runtime/SlateRHIRenderer/Private/SlateRHIResourceManager.cpp
+++ b/Engine/Source/Runtime/SlateRHIRenderer/Private/SlateRHIResourceManager.cpp
@@ -67,12 +67,28 @@ void FDynamicResourceMap::RemoveMaterialResource( UMaterialInterface* Material )
 }
 
 void FDynamicResourceMap::Empty()
+{
+	EmptyUTextureResources();
+	EmptyMaterialResources();
+	EmptyDynamicTextureResources();
+}
+
+void FDynamicResourceMap::EmptyDynamicTextureResources()
 {
 	NativeTextureMap.Empty();
+}
+
+void FDynamicResourceMap::EmptyUTextureResources()
+{
 	UTextureResourceMap.Empty();
+}
+
+void FDynamicResourceMap::EmptyMaterialResources()
+{
 	MaterialResourceMap.Empty();
 }
 
+
 void FDynamicResourceMap::ReleaseResources()
 {
 	for (TMap<FName, TSharedPtr<FSlateDynamicTextureResource> >::TIterator It(NativeTextureMap); It; ++It)
@@ -102,6 +118,8 @@ void FDynamicResourceMap::AddReferencedObjects(FReferenceCollector& Collector)
 
 FSlateRHIResourceManager::FSlateRHIResourceManager()
 {
+	FCoreDelegates::OnPreExit.AddRaw( this, &FSlateRHIResourceManager::OnAppExit );
+
 	MaxAltasedTextureSize = FIntPoint(256,256);
 	if( GIsEditor )
 	{
@@ -130,6 +148,8 @@ FSlateRHIResourceManager::FSlateRHIResourceManager()
 
 FSlateRHIResourceManager::~FSlateRHIResourceManager()
 {
+	FCoreDelegates::OnPreExit.RemoveAll( this );
+
 	if ( GIsRHIInitialized )
 	{
 		DeleteResources();
@@ -583,6 +603,15 @@ FSlateShaderResourceProxy* FSlateRHIResourceManager::GetMaterialResource(const F
 	return MaterialResource->SlateProxy;
 }
 
+void FSlateRHIResourceManager::OnAppExit()
+{
+	ReleaseResources();
+
+	FlushRenderingCommands();
+
+	DeleteResources();
+}
+
 bool FSlateRHIResourceManager::ContainsTexture( const FName& ResourceName ) const
 {
 	return ResourceMap.Contains( ResourceName );
@@ -599,7 +628,7 @@ void FSlateRHIResourceManager::ReleaseDynamicResource( const FSlateBrush& InBrus
 
 		UObject* ResourceObject = InBrush.GetResourceObject();
 
-		if( ResourceObject )
+		if( ResourceObject && DynamicResourceMap.GetNumObjectResources() > 0 )
 		{
 			TSharedPtr<FSlateUTextureResource> TextureResource = DynamicResourceMap.GetUTextureResource(Cast<UTexture2D>(ResourceObject));
 
@@ -628,7 +657,7 @@ void FSlateRHIResourceManager::ReleaseDynamicResource( const FSlateBrush& InBrus
 			}
 		
 		}
-		else
+		else if( !ResourceObject )
 		{
 			TSharedPtr<FSlateDynamicTextureResource> TextureResource = DynamicResourceMap.GetDynamicTextureResource(ResourceName);
 
diff --git a/Engine/Source/Runtime/SlateRHIRenderer/Private/SlateRHIResourceManager.h b/Engine/Source/Runtime/SlateRHIRenderer/Private/SlateRHIResourceManager.h
--- a/Engine/Source/Runtime/SlateRHIRenderer/Private/SlateRHIResourceManager.h
+++ b/Engine/Source/Runtime/SlateRHIRenderer/Private/SlateRHIResourceManager.h
@@ -26,8 +26,14 @@ public:
 
 	void Empty();
 
+	void EmptyUTextureResources();
+	void EmptyMaterialResources();
+	void EmptyDynamicTextureResources();
+
 	void ReleaseResources();
 
+	uint32 GetNumObjectResources() const { return UTextureResourceMap.Num() + MaterialResourceMap.Num(); }
+
 	/** FGCObject interface */
 	virtual void AddReferencedObjects(FReferenceCollector& Collector) override;
 
@@ -171,6 +177,11 @@ private:
 	 */
 	FSlateShaderResourceProxy* GetMaterialResource( const FSlateBrush& InBrush );
 
+	/**
+	 * Called when the application exists before the UObject system shuts down so we can free object resources
+	 */
+	void OnAppExit();
+
 private:
 	/** Map of all active dynamic resources being used by brushes */
 	FDynamicResourceMap DynamicResourceMap;
