Fixed a crash updating the font atlas

UE-4789 & UE-4793 - CRITICAL: Regression: UE4Editor_SlateRHIRenderer crash when building lighting

The rendering thread is disabled when building lighting, which made FSlateFontAtlasRHI::ConditionalUpdateTexture take the new code path. The two code paths should have ultimately been doing the same thing, but the local "Stride" variable was hiding the "Stride" member variable, causing a crash.

diff --git a/Engine/Source/Runtime/SlateRHIRenderer/Private/SlateRHIFontTexture.cpp b/Engine/Source/Runtime/SlateRHIRenderer/Private/SlateRHIFontTexture.cpp
--- a/Engine/Source/Runtime/SlateRHIRenderer/Private/SlateRHIFontTexture.cpp
+++ b/Engine/Source/Runtime/SlateRHIRenderer/Private/SlateRHIFontTexture.cpp
@@ -165,8 +165,8 @@ void FSlateFontAtlasRHI::ConditionalUpdateTexture()
 		{
 			FontTexture->InitResource();
 
-			uint32 Stride;
-			uint8* TempData = (uint8*)RHILockTexture2D( FontTexture->GetTypedResource(), 0, RLM_WriteOnly, Stride, false );
+			uint32 TextureStride;
+			uint8* TempData = (uint8*)RHILockTexture2D( FontTexture->GetTypedResource(), 0, RLM_WriteOnly, TextureStride, false );
 			FMemory::Memcpy( TempData, AtlasData.GetData(), Stride*AtlasWidth*AtlasHeight );
 			RHIUnlockTexture2D( FontTexture->GetTypedResource(),0,false );
 		}
@@ -179,8 +179,8 @@ void FSlateFontAtlasRHI::ConditionalUpdateTexture()
 			ENQUEUE_UNIQUE_RENDER_COMMAND_ONEPARAMETER( SlateUpdateFontTextureCommand,
 				FSlateFontAtlasRHI&, Atlas, *this,
 			{
-				uint32 Stride;
-				uint8* TempData = (uint8*)RHILockTexture2D( Atlas.FontTexture->GetTypedResource(), 0, RLM_WriteOnly, Stride, false );
+				uint32 TextureStride;
+				uint8* TempData = (uint8*)RHILockTexture2D( Atlas.FontTexture->GetTypedResource(), 0, RLM_WriteOnly, TextureStride, false );
 				FMemory::Memcpy( TempData, Atlas.AtlasData.GetData(), Atlas.Stride*Atlas.AtlasWidth*Atlas.AtlasHeight );
 				RHIUnlockTexture2D( Atlas.FontTexture->GetTypedResource(),0,false );
 			});
