Removed a redundant base class from FSlateFontTextureRHI

diff --git a/Engine/Source/Runtime/SlateRHIRenderer/Private/SlateRHIFontTexture.cpp b/Engine/Source/Runtime/SlateRHIRenderer/Private/SlateRHIFontTexture.cpp
--- a/Engine/Source/Runtime/SlateRHIRenderer/Private/SlateRHIFontTexture.cpp
+++ b/Engine/Source/Runtime/SlateRHIRenderer/Private/SlateRHIFontTexture.cpp
@@ -92,54 +92,6 @@ void FSlateFontTextureRHI::ReleaseDynamicRHI()
 	ShaderResource.SafeRelease();
 }
 
-void FSlateFontTextureRHI::ResizeTexture( uint32 InWidth, uint32 InHeight )
-{
-	if (Width != InWidth || Height != InHeight)
-	{
-		if (IsInRenderingThread())
-		{
-			Width = InWidth;
-			Height = InHeight;
-			UpdateRHI();
-		}
-		else
-		{
-			FIntPoint Dimensions(InWidth, InHeight);
-			ENQUEUE_UNIQUE_RENDER_COMMAND_TWOPARAMETER(SlateResizeFontTexture,
-			FSlateFontTextureRHI*, TextureRHIRef, this,
-			FIntPoint, InDimensions, Dimensions,
-			{
-				TextureRHIRef->Width = InDimensions.X;
-				TextureRHIRef->Height = InDimensions.Y;
-				TextureRHIRef->UpdateRHI();
-			});
-		}
-	}
-}
-
-void FSlateFontTextureRHI::UpdateTexture( const TArray<uint8>& InBytes )
-{
-	if (IsInRenderingThread())
-	{
-		uint32 Stride = 0;
-		void* TextureBuffer = RHILockTexture2D(GetTypedResource(), 0, RLM_WriteOnly, Stride, false);
-		FMemory::Memcpy(TextureBuffer, InBytes.GetData(), InBytes.Num());
-		RHIUnlockTexture2D(GetTypedResource(), 0, false);
-	}
-	else
-	{
-		ENQUEUE_UNIQUE_RENDER_COMMAND_TWOPARAMETER(SlateUpdateFontTexture,
-		FSlateFontTextureRHI*, TextureRHIRef, this,
-		const TArray<uint8>&, TextureData, InBytes,
-		{
-			uint32 Stride = 0;
-			void* TextureBuffer = RHILockTexture2D(TextureRHIRef->GetTypedResource(), 0, RLM_WriteOnly, Stride, false);
-			FMemory::Memcpy(TextureBuffer, TextureData.GetData(), TextureData.Num());
-			RHIUnlockTexture2D(TextureRHIRef->GetTypedResource(), 0, false);
-		});
-	}
-}
-
 FSlateFontAtlasRHI::FSlateFontAtlasRHI( uint32 Width, uint32 Height )
 	: FSlateFontAtlas( Width, Height ) 
 	, FontTexture( new FSlateFontTextureRHI( Width, Height ) )
diff --git a/Engine/Source/Runtime/SlateRHIRenderer/Private/SlateRHIFontTexture.h b/Engine/Source/Runtime/SlateRHIRenderer/Private/SlateRHIFontTexture.h
--- a/Engine/Source/Runtime/SlateRHIRenderer/Private/SlateRHIFontTexture.h
+++ b/Engine/Source/Runtime/SlateRHIRenderer/Private/SlateRHIFontTexture.h
@@ -8,7 +8,7 @@
  * ensure all characters in the font texture exist if the rendering resource has to be recreated 
  * between caching new characters
  */
-class FSlateFontTextureRHI : public TSlateTexture<FTexture2DRHIRef>, public FSlateUpdatableTexture, public FTextureResource
+class FSlateFontTextureRHI : public TSlateTexture<FTexture2DRHIRef>, public FTextureResource
 {
 public:
 	/** Constructor.  Initializes the texture
@@ -22,12 +22,6 @@ public:
 	virtual uint32 GetWidth() const override { return Width; }
 	virtual uint32 GetHeight() const override { return Height; }
 
-	/** FSlateUpdatableTexture interface */
-	virtual FSlateShaderResource* GetSlateResource() override { return this; }
-	virtual FRenderResource* GetRenderResource() override { return this; }
-	virtual void ResizeTexture( uint32 InWidth, uint32 InHeight ) override;
-	virtual void UpdateTexture( const TArray<uint8>& InBytes ) override;
-
 	/** FTextureResource interface */
 	virtual uint32 GetSizeX() const override { return Width; }
 	virtual uint32 GetSizeY() const override { return Height; }
