Fix for UE-5704: Resolution not restoring correctly after alt-tabbing out of and back in to a native fullscreen mode.

diff --git a/Engine/Source/Runtime/Engine/Public/UnrealEngine.h b/Engine/Source/Runtime/Engine/Public/UnrealEngine.h
--- a/Engine/Source/Runtime/Engine/Public/UnrealEngine.h
+++ b/Engine/Source/Runtime/Engine/Public/UnrealEngine.h
@@ -432,6 +432,12 @@ struct ENGINE_API FSystemResolution
 	{
 
 	}
+
+	void ForceRefresh()
+	{
+		RequestResolutionChange(ResX, ResY, WindowMode);
+		bForceRefresh = true;
+	}
 };
 
 ENGINE_API extern FSystemResolution GSystemResolution;
\ No newline at end of file
diff --git a/Engine/Source/Runtime/SlateRHIRenderer/Private/SlateRHIRenderer.cpp b/Engine/Source/Runtime/SlateRHIRenderer/Private/SlateRHIRenderer.cpp
--- a/Engine/Source/Runtime/SlateRHIRenderer/Private/SlateRHIRenderer.cpp
+++ b/Engine/Source/Runtime/SlateRHIRenderer/Private/SlateRHIRenderer.cpp
@@ -359,7 +359,7 @@ void FSlateRHIRenderer::RestoreSystemResolution(const TSharedRef<SWindow> InWind
 		// Force the window system to resize the active viewport, even though nothing might have appeared to change.
 		// On windows, DXGI might change the window resolution behind our backs when we alt-tab out. This will make
 		// sure that we are actually in the resolution we think we are.
-		GSystemResolution.bForceRefresh = true;
+		GSystemResolution.ForceRefresh();
 	}
 }
 
