Made font atlas pages in the editor 2048x2048

This had been done for texture atlas pages, but missed for font atlas pages.

diff --git a/Engine/Source/Runtime/SlateRHIRenderer/Private/SlateRHIRendererModule.cpp b/Engine/Source/Runtime/SlateRHIRenderer/Private/SlateRHIRendererModule.cpp
--- a/Engine/Source/Runtime/SlateRHIRenderer/Private/SlateRHIRendererModule.cpp
+++ b/Engine/Source/Runtime/SlateRHIRenderer/Private/SlateRHIRendererModule.cpp
@@ -8,12 +8,18 @@ class FSlateRHIFontAtlasFactory : public ISlateFontAtlasFactory
 public:
 	FSlateRHIFontAtlasFactory()
 	{
-		/** Size of each font texture, width and height */
-		AtlasSize = 1024;
-		if (!GIsEditor && GConfig)
+		if (GIsEditor)
 		{
-			GConfig->GetInt(TEXT("SlateRenderer"), TEXT("FontAtlasSize"), AtlasSize, GEngineIni);
-			AtlasSize = FMath::Clamp(AtlasSize, 0, 2048);
+			AtlasSize = 2048;
+		}
+		else
+		{
+			AtlasSize = 1024;
+			if (GConfig)
+			{
+				GConfig->GetInt(TEXT("SlateRenderer"), TEXT("FontAtlasSize"), AtlasSize, GEngineIni);
+				AtlasSize = FMath::Clamp(AtlasSize, 0, 2048);
+			}
 		}
 	}
 
@@ -21,12 +27,13 @@ public:
 	{
 	}
 
-
 	virtual TSharedRef<FSlateFontAtlas> CreateFontAtlas() const override
 	{
 		return MakeShareable(new FSlateFontAtlasRHI(AtlasSize, AtlasSize));
 	}
+
 private:
+	/** Size of each font texture, width and height */
 	int32 AtlasSize;
 };
 
