Fix resizing issues on html5 for now.

#html5

diff --git a/Engine/Source/Runtime/SlateRHIRenderer/Private/SlateRHIRenderer.cpp b/Engine/Source/Runtime/SlateRHIRenderer/Private/SlateRHIRenderer.cpp
--- a/Engine/Source/Runtime/SlateRHIRenderer/Private/SlateRHIRenderer.cpp
+++ b/Engine/Source/Runtime/SlateRHIRenderer/Private/SlateRHIRenderer.cpp
@@ -359,7 +359,10 @@ void FSlateRHIRenderer::RestoreSystemResolution(const TSharedRef<SWindow> InWind
 		// Force the window system to resize the active viewport, even though nothing might have appeared to change.
 		// On windows, DXGI might change the window resolution behind our backs when we alt-tab out. This will make
 		// sure that we are actually in the resolution we think we are.
+#if !PLATFORM_HTML5
+		// @todo: fixme for HTML5. 
 		GSystemResolution.ForceRefresh();
+#endif 
 	}
 }
 
