Expanded Slate window transparency support to include per-pixel alpha blending.

Currently this is only implemented for Windows (DirectX). Mac (OpenGL) support is in the pipeline and should be added soon.

diff --git a/Engine/Source/Developer/StandaloneRenderer/Private/Windows/D3D/SlateD3DRenderer.cpp b/Engine/Source/Developer/StandaloneRenderer/Private/Windows/D3D/SlateD3DRenderer.cpp
--- a/Engine/Source/Developer/StandaloneRenderer/Private/Windows/D3D/SlateD3DRenderer.cpp
+++ b/Engine/Source/Developer/StandaloneRenderer/Private/Windows/D3D/SlateD3DRenderer.cpp
@@ -358,7 +358,13 @@ void FSlateD3DRenderer::DrawWindows( FSlateDrawBuffer& InWindowDrawBuffer )
 
 			ID3D11RenderTargetView* RTV = Viewport->RenderTargetView;
 			ID3D11DepthStencilView* DSV = Viewport->DepthStencilView;
-
+			
+#if ALPHA_BLENDED_WINDOWS
+			if ( WindowToDraw->GetTransparencySupport() == EWindowTransparency::PerPixel )
+			{
+				GD3DDeviceContext->ClearRenderTargetView(RTV, &FLinearColor::Transparent.R);
+			}
+#endif
 			GD3DDeviceContext->OMSetRenderTargets( 1, &RTV, NULL );
 
 			RenderingPolicy->DrawElements( ViewMatrix*Viewport->ProjectionMatrix, ElementList.GetRenderBatches() );
diff --git a/Engine/Source/Developer/StandaloneRenderer/Private/Windows/D3D/SlateD3DRenderingPolicy.cpp b/Engine/Source/Developer/StandaloneRenderer/Private/Windows/D3D/SlateD3DRenderingPolicy.cpp
--- a/Engine/Source/Developer/StandaloneRenderer/Private/Windows/D3D/SlateD3DRenderingPolicy.cpp
+++ b/Engine/Source/Developer/StandaloneRenderer/Private/Windows/D3D/SlateD3DRenderingPolicy.cpp
@@ -83,11 +83,14 @@ void FSlateD3D11RenderingPolicy::InitResources()
 	BlendDesc.AlphaToCoverageEnable = false;
 	BlendDesc.IndependentBlendEnable = false;
 	BlendDesc.RenderTarget[0].SrcBlend = D3D11_BLEND_SRC_ALPHA;
-	BlendDesc.RenderTarget[0].SrcBlendAlpha = D3D11_BLEND_SRC_ALPHA;
 	BlendDesc.RenderTarget[0].BlendOp = D3D11_BLEND_OP_ADD;
+
+	// dest.a = 1-(1-dest.a)*src.a + dest.a
+	BlendDesc.RenderTarget[0].SrcBlendAlpha = D3D11_BLEND_INV_DEST_ALPHA;
+	BlendDesc.RenderTarget[0].DestBlendAlpha = D3D11_BLEND_ONE;
 	BlendDesc.RenderTarget[0].BlendOpAlpha = D3D11_BLEND_OP_ADD;
+
 	BlendDesc.RenderTarget[0].DestBlend = D3D11_BLEND_INV_SRC_ALPHA;
-	BlendDesc.RenderTarget[0].DestBlendAlpha = D3D11_BLEND_INV_SRC_ALPHA;
 	BlendDesc.RenderTarget[0].RenderTargetWriteMask = D3D11_COLOR_WRITE_ENABLE_ALL;
 	BlendDesc.RenderTarget[0].BlendEnable = true;
 
diff --git a/Engine/Source/Runtime/Core/Private/Windows/WindowsApplication.cpp b/Engine/Source/Runtime/Core/Private/Windows/WindowsApplication.cpp
--- a/Engine/Source/Runtime/Core/Private/Windows/WindowsApplication.cpp
+++ b/Engine/Source/Runtime/Core/Private/Windows/WindowsApplication.cpp
@@ -21,6 +21,7 @@
 	#include <objbase.h>
 	#include <SetupApi.h>
 	#include <devguid.h>
+	#include <dwmapi.h>
 
 // This might not be defined by Windows when maintaining backwards-compatibility to pre-Vista builds
 #ifndef WM_MOUSEHWHEEL
@@ -606,6 +607,14 @@ EWindowTitleAlignment::Type FWindowsApplication::GetWindowTitleAlignment() const
 
 	return EWindowTitleAlignment::Left;
 }
+	
+EWindowTransparency FWindowsApplication::GetWindowTransparencySupport() const
+{
+	BOOL bIsCompositionEnabled = FALSE;
+	::DwmIsCompositionEnabled(&bIsCompositionEnabled);
+
+	return bIsCompositionEnabled ? EWindowTransparency::PerPixel : EWindowTransparency::PerWindow;
+}
 
 /** All WIN32 messages sent to our app go here; this method simply passes them on */
 LRESULT CALLBACK FWindowsApplication::AppWndProc(HWND hwnd, uint32 msg, WPARAM wParam, LPARAM lParam)
@@ -890,6 +899,14 @@ int32 FWindowsApplication::ProcessMessage( HWND hwnd, uint32 msg, WPARAM wParam,
 				}
 			}
 			break;
+			
+#if WINVER > 0x502
+		case WM_DWMCOMPOSITIONCHANGED:
+			{
+				DeferMessage( CurrentNativeEventWindowPtr, hwnd, msg, wParam, lParam );
+			}
+			break;
+#endif
 
 			// Window focus and activation
 		case WM_ACTIVATE:
@@ -1588,6 +1605,14 @@ int32 FWindowsApplication::ProcessDeferredMessage( const FDeferredWindowsMessage
 				}
 			}
 			break;
+
+#if WINVER > 0x502
+		case WM_DWMCOMPOSITIONCHANGED:
+			{
+				CurrentNativeEventWindowPtr->OnTransparencySupportChanged(GetWindowTransparencySupport());
+			}
+			break;
+#endif
 		}
 	}
 
diff --git a/Engine/Source/Runtime/Core/Private/Windows/WindowsWindow.cpp b/Engine/Source/Runtime/Core/Private/Windows/WindowsWindow.cpp
--- a/Engine/Source/Runtime/Core/Private/Windows/WindowsWindow.cpp
+++ b/Engine/Source/Runtime/Core/Private/Windows/WindowsWindow.cpp
@@ -59,14 +59,25 @@ void FWindowsWindow::Initialize( FWindowsApplication* const Application, const T
 	int32 ClientHeight = FMath::TruncToInt( HeightInitial );
 	int32 WindowWidth = ClientWidth;
 	int32 WindowHeight = ClientHeight;
+	const bool bApplicationSupportsPerPixelBlending = Application->GetWindowTransparencySupport() == EWindowTransparency::PerPixel;
 
 	if( !Definition->HasOSWindowBorder )
 	{
 		WindowExStyle = WS_EX_WINDOWEDGE;
-		if (Definition->SupportsTransparency)
+
+		if( Definition->TransparencySupport == EWindowTransparency::PerWindow )
 		{
 			WindowExStyle |= WS_EX_LAYERED;
 		}
+#if ALPHA_BLENDED_WINDOWS
+		else if( Definition->TransparencySupport == EWindowTransparency::PerPixel )
+		{
+			if( bApplicationSupportsPerPixelBlending )
+			{
+				WindowExStyle |= WS_EX_COMPOSITED;
+			}
+		}
+#endif
 
 		WindowStyle = WS_POPUP | WS_CLIPCHILDREN | WS_CLIPSIBLINGS;
 		if ( Definition->AppearsInTaskbar )
@@ -134,7 +145,7 @@ void FWindowsWindow::Initialize( FWindowsApplication* const Application, const T
 		return;
 	}
 
-	if ( Definition->SupportsTransparency )
+	if ( Definition->TransparencySupport == EWindowTransparency::PerWindow )
 	{
 		SetOpacity( Definition->Opacity );
 	}
@@ -149,7 +160,16 @@ void FWindowsWindow::Initialize( FWindowsApplication* const Application, const T
 
 		const BOOL bEnableNCPaint = false;
 		verify(SUCCEEDED(DwmSetWindowAttribute(HWnd, DWMWA_ALLOW_NCPAINT, &bEnableNCPaint, sizeof(bEnableNCPaint))));
+
+	#if ALPHA_BLENDED_WINDOWS
+		if ( bApplicationSupportsPerPixelBlending && Definition->TransparencySupport == EWindowTransparency::PerPixel )
+		{
+			MARGINS Margins = {-1};
+			verify(SUCCEEDED(::DwmExtendFrameIntoClientArea(HWnd, &Margins)));
+		}
+	#endif
 	}
+
 #endif	// WINVER
 
 	// No region for non regular windows or windows displaying the os window border
@@ -200,6 +220,31 @@ HWND FWindowsWindow::GetHWnd() const
 	return HWnd;
 }
 
+void FWindowsWindow::OnTransparencySupportChanged(EWindowTransparency NewTransparency)
+{
+#if ALPHA_BLENDED_WINDOWS
+	if ( Definition->TransparencySupport == EWindowTransparency::PerPixel )
+	{
+		const auto Style = GetWindowLong( HWnd, GWL_EXSTYLE );
+
+		if( NewTransparency == EWindowTransparency::PerPixel )
+		{
+			SetWindowLong( HWnd, GWL_EXSTYLE, Style | WS_EX_COMPOSITED );
+
+			MARGINS Margins = {-1};
+			verify(SUCCEEDED(::DwmExtendFrameIntoClientArea(HWnd, &Margins)));
+		}
+		else
+		{
+			SetWindowLong( HWnd, GWL_EXSTYLE, Style & ~WS_EX_COMPOSITED );
+		}
+
+		// Must call SWP_FRAMECHANGED when updating the style attribute of a window in order to update internal caches (according to MSDN)
+		SetWindowPos( HWnd, NULL, 0, 0, 0, 0, SWP_FRAMECHANGED | SWP_NOACTIVATE | SWP_NOMOVE | SWP_NOOWNERZORDER | SWP_NOREDRAW | SWP_NOSIZE | SWP_NOSENDCHANGING | SWP_NOZORDER );
+	}
+#endif
+}
+
 HRGN FWindowsWindow::MakeWindowRegionObject() const
 {
 	HRGN Region;
@@ -212,7 +257,14 @@ HRGN FWindowsWindow::MakeWindowRegionObject() const
 		}
 		else
 		{
-			if( Definition->CornerRadius > 0 )
+			const bool bUseCornerRadius  = 
+#if ALPHA_BLENDED_WINDOWS
+				// Corner radii cause DWM window composition blending to fail, so we always set regions to full size rectangles
+				Definition->TransparencySupport != EWindowTransparency::PerPixel &&
+#endif
+				Definition->CornerRadius > 0;
+
+			if( bUseCornerRadius )
 			{
 				// @todo mac: Corner radius not applied on Mac platform yet
 
diff --git a/Engine/Source/Runtime/Core/Public/GenericPlatform/GenericApplication.h b/Engine/Source/Runtime/Core/Public/GenericPlatform/GenericApplication.h
--- a/Engine/Source/Runtime/Core/Public/GenericPlatform/GenericApplication.h
+++ b/Engine/Source/Runtime/Core/Public/GenericPlatform/GenericApplication.h
@@ -438,6 +438,11 @@ public:
 		return EWindowTitleAlignment::Left;
 	}
 
+	virtual EWindowTransparency GetWindowTransparencySupport() const
+	{
+		return EWindowTransparency::None;
+	}
+
 	virtual void DestroyApplication() { }
 
 	/** Function to return the current implementation of the ForceFeedback system */
diff --git a/Engine/Source/Runtime/Core/Public/GenericPlatform/GenericWindowDefinition.h b/Engine/Source/Runtime/Core/Public/GenericPlatform/GenericWindowDefinition.h
--- a/Engine/Source/Runtime/Core/Public/GenericPlatform/GenericWindowDefinition.h
+++ b/Engine/Source/Runtime/Core/Public/GenericPlatform/GenericWindowDefinition.h
@@ -5,6 +5,25 @@
 #include "HAL/Platform.h"
 #include "GenericApplicationMessageHandler.h"
 
+#ifndef ALPHA_BLENDED_WINDOWS
+	#define ALPHA_BLENDED_WINDOWS IS_PROGRAM || WITH_EDITOR
+#endif
+
+/** Enumeration to specify different transparency options for SWindows */
+enum class EWindowTransparency
+{
+	/** Value indicating that a window does not support transparency */
+	None,
+
+	/** Value indicating that a window supports transparency at the window level (one opacity applies to the entire window) */
+	PerWindow,
+
+#if ALPHA_BLENDED_WINDOWS
+	/** Value indicating that a window supports per-pixel alpha blended transparency */
+	PerPixel,
+#endif
+};
+
 struct CORE_API FGenericWindowDefinition
 {
 	float XDesiredPositionOnScreen;
@@ -13,8 +32,9 @@ struct CORE_API FGenericWindowDefinition
 	float WidthDesiredOnScreen;
 	float HeightDesiredOnScreen;
 
+	EWindowTransparency TransparencySupport;
+
 	bool HasOSWindowBorder;
-	bool SupportsTransparency;
 	bool AppearsInTaskbar;
 	bool IsTopmostWindow;
 	bool AcceptsInput;
diff --git a/Engine/Source/Runtime/Core/Public/Linux/LinuxApplication.h b/Engine/Source/Runtime/Core/Public/Linux/LinuxApplication.h
--- a/Engine/Source/Runtime/Core/Public/Linux/LinuxApplication.h
+++ b/Engine/Source/Runtime/Core/Public/Linux/LinuxApplication.h
@@ -51,6 +51,11 @@ public:
 	//	X
 	virtual bool TryCalculatePopupWindowPosition( const FPlatformRect& InAnchor, const FVector2D& InSize, const EPopUpOrientation::Type Orientation, /*OUT*/ FVector2D* const CalculatedPopUpPosition ) const override;
 
+	virtual EWindowTransparency GetWindowTransparencySupport() const override
+	{
+		return EWindowTransparency::PerPixel;
+	}
+
 	void AddPendingEvent( SDL_Event event );
 
 	void OnMouseCursorLock( bool bLockEnabled );
diff --git a/Engine/Source/Runtime/Core/Public/Mac/MacApplication.h b/Engine/Source/Runtime/Core/Public/Mac/MacApplication.h
--- a/Engine/Source/Runtime/Core/Public/Mac/MacApplication.h
+++ b/Engine/Source/Runtime/Core/Public/Mac/MacApplication.h
@@ -52,6 +52,11 @@ public:
 	{
 		return EWindowTitleAlignment::Center;
 	}
+	
+	virtual EWindowTransparency GetWindowTransparencySupport() const override
+	{
+		return EWindowTransparency::PerPixel;
+	}
 
 	virtual ITextInputMethodSystem *GetTextInputMethodSystem() override
 	{
diff --git a/Engine/Source/Runtime/Core/Public/Windows/WindowsApplication.h b/Engine/Source/Runtime/Core/Public/Windows/WindowsApplication.h
--- a/Engine/Source/Runtime/Core/Public/Windows/WindowsApplication.h
+++ b/Engine/Source/Runtime/Core/Public/Windows/WindowsApplication.h
@@ -259,6 +259,8 @@ public:
 
 	virtual EWindowTitleAlignment::Type GetWindowTitleAlignment() const override;
 	
+	virtual EWindowTransparency GetWindowTransparencySupport() const override;
+
 	/** Function to return the current implementation of the ForceFeedback system */
 	DEPRECATED(4.7, "Please use GetInputInterface()")
 	virtual IForceFeedbackSystem* GetForceFeedbackSystem() override
diff --git a/Engine/Source/Runtime/Core/Public/Windows/WindowsWindow.h b/Engine/Source/Runtime/Core/Public/Windows/WindowsWindow.h
--- a/Engine/Source/Runtime/Core/Public/Windows/WindowsWindow.h
+++ b/Engine/Source/Runtime/Core/Public/Windows/WindowsWindow.h
@@ -120,6 +120,11 @@ public:
 	virtual HRESULT STDCALL DragOver( ::DWORD KeyState, POINTL CursorPosition, __RPC__inout ::DWORD *CursorEffect) override;
 	virtual HRESULT STDCALL DragLeave( void ) override;
 	virtual HRESULT STDCALL Drop( __RPC__in_opt IDataObject *DataObjectPointer, ::DWORD KeyState, POINTL CursorPosition, __RPC__inout ::DWORD *CursorEffect) override;
+	
+public:
+
+	/** Called by the owning application when the level of transparency support has changed */
+	void OnTransparencySupportChanged(EWindowTransparency NewTransparency);
 
 private:
 
diff --git a/Engine/Source/Runtime/Slate/Private/Framework/Application/MenuStack.cpp b/Engine/Source/Runtime/Slate/Private/Framework/Application/MenuStack.cpp
--- a/Engine/Source/Runtime/Slate/Private/Framework/Application/MenuStack.cpp
+++ b/Engine/Source/Runtime/Slate/Private/Framework/Application/MenuStack.cpp
@@ -151,7 +151,7 @@ TSharedRef<SWindow> FMenuStack::PushMenu( const TSharedRef<SWindow>& ParentWindo
 	}
 
 	// Start pop-up windows out transparent, then fade them in over time
-	const bool bUseTransparency = bAllowAnimations ? true : false;
+	const EWindowTransparency Transparency(bAllowAnimations ? EWindowTransparency::PerWindow : EWindowTransparency::None);
 	const float InitialWindowOpacity = bAllowAnimations ? 0.0f : 1.0f;
 	const float TargetWindowOpacity = 1.0f;
 
@@ -169,7 +169,7 @@ TSharedRef<SWindow> FMenuStack::PushMenu( const TSharedRef<SWindow>& ParentWindo
 		.AutoCenter( EAutoCenter::None )
 		.ClientSize( ExpectedSize )
 		.InitialOpacity(InitialWindowOpacity)
-		.SupportsTransparency( bUseTransparency )
+		.SupportsTransparency( Transparency )
 		.FocusWhenFirstShown( bFocusImmediately )
 		.ActivateWhenFirstShown( bFocusImmediately )
 		[
diff --git a/Engine/Source/Runtime/Slate/Private/Framework/Application/SlateApplication.cpp b/Engine/Source/Runtime/Slate/Private/Framework/Application/SlateApplication.cpp
--- a/Engine/Source/Runtime/Slate/Private/Framework/Application/SlateApplication.cpp
+++ b/Engine/Source/Runtime/Slate/Private/Framework/Application/SlateApplication.cpp
@@ -1522,7 +1522,7 @@ TSharedRef< FGenericWindow > FSlateApplication::MakeWindow( TSharedRef<SWindow>
 	Definition->YDesiredPositionOnScreen = Position.Y;
 
 	Definition->HasOSWindowBorder = InSlateWindow->HasOSWindowBorder();
-	Definition->SupportsTransparency = InSlateWindow->SupportsTransparency();
+	Definition->TransparencySupport = InSlateWindow->GetTransparencySupport();
 	Definition->AppearsInTaskbar = InSlateWindow->AppearsInTaskbar();
 	Definition->IsTopmostWindow = InSlateWindow->IsTopmostWindow();
 	Definition->AcceptsInput = InSlateWindow->AcceptsInput();
diff --git a/Engine/Source/Runtime/Slate/Public/Framework/Application/SlateApplication.h b/Engine/Source/Runtime/Slate/Public/Framework/Application/SlateApplication.h
--- a/Engine/Source/Runtime/Slate/Public/Framework/Application/SlateApplication.h
+++ b/Engine/Source/Runtime/Slate/Public/Framework/Application/SlateApplication.h
@@ -978,6 +978,11 @@ public:
 	virtual EVisibility GetSoftwareCursorVis() const override;	
 	virtual TSharedPtr<SWidget> GetKeyboardFocusedWidget() const override;
 
+	virtual EWindowTransparency GetWindowTransparencySupport() const override
+	{
+		return PlatformApplication->GetWindowTransparencySupport();
+	}
+
 protected:
 
 	virtual TSharedPtr< SWidget > GetMouseCaptorImpl() const override;
diff --git a/Engine/Source/Runtime/SlateCore/Private/Widgets/SWindow.cpp b/Engine/Source/Runtime/SlateCore/Private/Widgets/SWindow.cpp
--- a/Engine/Source/Runtime/SlateCore/Private/Widgets/SWindow.cpp
+++ b/Engine/Source/Runtime/SlateCore/Private/Widgets/SWindow.cpp
@@ -185,7 +185,7 @@ void SWindow::Construct(const FArguments& InArgs)
 
 	this->Title = InArgs._Title;
 	this->bDragAnywhere = InArgs._bDragAnywhere;
-	this->bIsTransparent = InArgs._SupportsTransparency;
+	this->TransparencySupport = InArgs._SupportsTransparency.Value;
 	this->Opacity = InArgs._InitialOpacity;
 	this->bInitiallyMaximized = InArgs._IsInitiallyMaximized;
 	this->SizingRule = InArgs._SizingRule;
@@ -310,7 +310,7 @@ TSharedRef<SWindow> SWindow::MakeNotificationWindow()
 		.IsPopupWindow( true )
 		.CreateTitleBar( false )
 		.SizingRule( ESizingRule::Autosized )
-		.SupportsTransparency( true )
+		.SupportsTransparency( EWindowTransparency::PerWindow )
 		.InitialOpacity( 0.0f )
 		.FocusWhenFirstShown( false )
 		.ActivateWhenFirstShown( false );
@@ -329,11 +329,11 @@ TSharedRef<SWindow> SWindow::MakeToolTipWindow()
 	TSharedRef<SWindow> NewWindow = SNew( SWindow )
 		.IsPopupWindow( true )
 		.SizingRule( ESizingRule::Autosized )
+		.SupportsTransparency( EWindowTransparency::PerWindow )
 		.FocusWhenFirstShown( false )
 		.ActivateWhenFirstShown( false );
 	NewWindow->bIsToolTipWindow = true;
 	NewWindow->bIsTopmostWindow = true;
-	NewWindow->bIsTransparent = true;
 	NewWindow->Opacity = 0.0f;
 
 	// NOTE: These sizes are tweaked for SToolTip widgets (text wrap width of around 400 px)
@@ -350,12 +350,12 @@ TSharedRef<SWindow> SWindow::MakeCursorDecorator()
 	TSharedRef<SWindow> NewWindow = SNew( SWindow )
 		.IsPopupWindow( true )
 		.SizingRule( ESizingRule::Autosized )
+		.SupportsTransparency( EWindowTransparency::PerWindow )
 		.FocusWhenFirstShown( false )
 		.ActivateWhenFirstShown( false );
 	NewWindow->bIsToolTipWindow = true;
 	NewWindow->bIsTopmostWindow = true;
 	NewWindow->bIsCursorDecoratorWindow = true;
-	NewWindow->bIsTransparent = true;
 	NewWindow->Opacity = 1.0f;
 
 	return NewWindow;
@@ -1182,9 +1182,9 @@ float SWindow::GetOpacity() const
 	return Opacity;
 }
 
-bool SWindow::SupportsTransparency() const
+EWindowTransparency SWindow::GetTransparencySupport() const
 {
-	return bIsTransparent;
+	return TransparencySupport;
 }
 
 
@@ -1615,7 +1615,7 @@ SWindow::SWindow()
 	: bDragAnywhere( false )
 	, Opacity( 1.0f )
 	, SizingRule( ESizingRule::UserSized )
-	, bIsTransparent( false )
+	, TransparencySupport( EWindowTransparency::None )
 	, bIsPopupWindow( false )
 	, bIsToolTipWindow( false )
 	, bIsTopmostWindow( false )
diff --git a/Engine/Source/Runtime/SlateCore/Public/Application/SlateApplicationBase.h b/Engine/Source/Runtime/SlateCore/Public/Application/SlateApplicationBase.h
--- a/Engine/Source/Runtime/SlateCore/Public/Application/SlateApplicationBase.h
+++ b/Engine/Source/Runtime/SlateCore/Public/Application/SlateApplicationBase.h
@@ -166,6 +166,13 @@ public:
 	 */
 	void GetDisplayMetrics(FDisplayMetrics& OutDisplayMetrics) const { FDisplayMetrics::GetDisplayMetrics(OutDisplayMetrics); }
 
+	/**
+	 * Get the highest level of window transparency support currently enabled by this application
+	 *
+	 * @return Enumeration value specifying the level of transparency currently supported
+	 */
+	virtual EWindowTransparency GetWindowTransparencySupport() const = 0;
+
 	/**
 	 * Gets the widget that currently has keyboard focus, if any.
 	 *
diff --git a/Engine/Source/Runtime/SlateCore/Public/Widgets/SWindow.h b/Engine/Source/Runtime/SlateCore/Public/Widgets/SWindow.h
--- a/Engine/Source/Runtime/SlateCore/Public/Widgets/SWindow.h
+++ b/Engine/Source/Runtime/SlateCore/Public/Widgets/SWindow.h
@@ -57,6 +57,16 @@ namespace ESizingRule
 	};
 }
 
+/** Proxy structure to handle deprecated construction from bool */
+struct FWindowTransparency
+{
+	DEPRECATED(4.7, "Please now use the static initializers defined below.")
+	FWindowTransparency(bool bSupportsTransparency) : Value(bSupportsTransparency ? EWindowTransparency::PerWindow : EWindowTransparency::None) {}
+
+	FWindowTransparency(EWindowTransparency In) : Value(In) {}
+	
+	EWindowTransparency Value;
+};
 
 /**
  * SWindow is a platform-agnostic representation of a top-level window.
@@ -73,7 +83,7 @@ public:
 		, _AutoCenter( EAutoCenter::PreferredWorkArea )
 		, _ScreenPosition( FVector2D::ZeroVector )
 		, _ClientSize( FVector2D::ZeroVector )
-		, _SupportsTransparency( false )
+		, _SupportsTransparency( EWindowTransparency::None )
 		, _InitialOpacity( 1.0f )
 		, _IsInitiallyMaximized( false )
 		, _SizingRule( ESizingRule::UserSized )
@@ -109,7 +119,7 @@ public:
 		SLATE_ARGUMENT( FVector2D, ClientSize )
 
 		/** Should this window support transparency */
-		SLATE_ARGUMENT( bool, SupportsTransparency )
+		SLATE_ARGUMENT( FWindowTransparency, SupportsTransparency )
 
 		/** The initial opacity of the window */
 		SLATE_ARGUMENT( float, InitialOpacity )
@@ -484,8 +494,8 @@ public:
 	/** @return the window's current opacity */
 	float GetOpacity() const;
 
-	/** @return true if the window supports transparency */
-	bool SupportsTransparency() const;
+	/** @return the level of transparency supported by this window */
+	EWindowTransparency GetTransparencySupport() const;
 
 	/** @return A String representation of the widget */
 	virtual FString ToString() const;
@@ -736,8 +746,8 @@ protected:
 	/** How to auto center the window */
 	EAutoCenter::Type AutoCenterRule;
 
-	/** true if the window supports transparency */
-	bool bIsTransparent : 1;
+	/** Transparency setting for this window */
+	EWindowTransparency TransparencySupport;
 
 	/** True if this is a pop up window */
 	bool bIsPopupWindow : 1;
diff --git a/Engine/Source/Runtime/SlateRHIRenderer/Private/SlateRHIRenderer.cpp b/Engine/Source/Runtime/SlateRHIRenderer/Private/SlateRHIRenderer.cpp
--- a/Engine/Source/Runtime/SlateRHIRenderer/Private/SlateRHIRenderer.cpp
+++ b/Engine/Source/Runtime/SlateRHIRenderer/Private/SlateRHIRenderer.cpp
@@ -258,12 +258,18 @@ void FSlateRHIRenderer::CreateViewport( const TSharedRef<SWindow> Window )
 		NewInfo->DesiredWidth = Width;
 		NewInfo->DesiredHeight = Height;
 		NewInfo->ProjectionMatrix = CreateProjectionMatrix( Width, Height );
+#if ALPHA_BLENDED_WINDOWS		
+		if( Window->GetTransparencySupport() == EWindowTransparency::PerPixel )
+		{
+			NewInfo->PixelFormat = PF_B8G8R8A8;
+		}
+#endif
 
 		// Sanity check dimensions
 		checkf(Width <= MAX_VIEWPORT_SIZE && Height <= MAX_VIEWPORT_SIZE, TEXT("Invalid window with Width=%u and Height=%u"), Width, Height);
 
 		bool bFullscreen = IsViewportFullscreen( *Window );
-		NewInfo->ViewportRHI = RHICreateViewport( NewInfo->OSWindow, Width, Height, bFullscreen, EPixelFormat::PF_Unknown );
+		NewInfo->ViewportRHI = RHICreateViewport( NewInfo->OSWindow, Width, Height, bFullscreen, NewInfo->PixelFormat );
 		NewInfo->bFullscreen = bFullscreen;
 
 		WindowToViewportInfo.Add( &Window.Get(), NewInfo );
@@ -316,7 +322,7 @@ void FSlateRHIRenderer::ConditionalResizeViewport( FViewportInfo* ViewInfo, uint
 		}
 		else
 		{
-			ViewInfo->ViewportRHI = RHICreateViewport(ViewInfo->OSWindow, NewWidth, NewHeight, bFullscreen, EPixelFormat::PF_Unknown);
+			ViewInfo->ViewportRHI = RHICreateViewport(ViewInfo->OSWindow, NewWidth, NewHeight, bFullscreen, ViewInfo->PixelFormat);
 		}
 
 		// Safe to call here as the rendering thread has been suspended: game thread == render thread!
@@ -387,7 +393,7 @@ void FSlateRHIRenderer::OnWindowDestroyed( const TSharedRef<SWindow>& InWindow )
 }
 
 /** Draws windows from a FSlateDrawBuffer on the render thread */
-void FSlateRHIRenderer::DrawWindow_RenderThread(FRHICommandListImmediate& RHICmdList, const FViewportInfo& ViewportInfo, const FSlateWindowElementList& WindowElementList, bool bLockToVsync)
+void FSlateRHIRenderer::DrawWindow_RenderThread(FRHICommandListImmediate& RHICmdList, const FViewportInfo& ViewportInfo, const FSlateWindowElementList& WindowElementList, bool bLockToVsync, bool bClear)
 {
 	SCOPED_DRAW_EVENT(RHICmdList, SlateUI);
 
@@ -425,6 +431,11 @@ void FSlateRHIRenderer::DrawWindow_RenderThread(FRHICommandListImmediate& RHICmd
 			SetRenderTarget(RHICmdList, BackBuffer, FTextureRHIRef());
 		}
 
+		if (bClear)
+		{
+			RHICmdList.Clear( true, FLinearColor(ForceInitToZero), false, 0.0f, true, 0x00, FIntRect());
+		}
+
 #if DEBUG_OVERDRAW
 		RHIClear(true, FLinearColor::Black, false, 0.0f, true, 0x00, FIntRect());
 #endif
@@ -618,12 +629,18 @@ void FSlateRHIRenderer::DrawWindows_Private( FSlateDrawBuffer& WindowDrawBuffer
 						FSlateWindowElementList* WindowElementList;
 						SWindow* SlateWindow;
 						bool bLockToVsync;
+						bool bClear;
 					} Params;
 
 					Params.Renderer = this;
 					Params.ViewportInfo = ViewInfo;
 					Params.WindowElementList = &ElementList;
 					Params.bLockToVsync = bLockToVsync;
+#if ALPHA_BLENDED_WINDOWS
+					Params.bClear = Window->GetTransparencySupport() == EWindowTransparency::PerPixel;
+#else
+					Params.bClear = false;
+#endif
 
 					// NOTE: We pass a raw pointer to the SWindow so that we don't have to use a thread-safe weak pointer in
 					// the FSlateWindowElementList structure
@@ -635,7 +652,7 @@ void FSlateRHIRenderer::DrawWindows_Private( FSlateDrawBuffer& WindowDrawBuffer
 						ENQUEUE_UNIQUE_RENDER_COMMAND_ONEPARAMETER(SlateDrawWindowsCommand,
 							FSlateDrawWindowCommandParams, Params, Params,
 							{
-								Params.Renderer->DrawWindow_RenderThread(RHICmdList, *Params.ViewportInfo, *Params.WindowElementList, Params.bLockToVsync);
+								Params.Renderer->DrawWindow_RenderThread(RHICmdList, *Params.ViewportInfo, *Params.WindowElementList, Params.bLockToVsync, Params.bClear);
 							});
 					}
 
@@ -1119,7 +1136,7 @@ void* FSlateRHIRenderer::GetViewportResource( const SWindow& Window )
 
 			const bool bFullscreen = IsViewportFullscreen( Window );
 
-			ViewportInfo->ViewportRHI = RHICreateViewport( ViewportInfo->OSWindow, ViewportInfo->Width, ViewportInfo->Height, bFullscreen, EPixelFormat::PF_Unknown );
+			ViewportInfo->ViewportRHI = RHICreateViewport( ViewportInfo->OSWindow, ViewportInfo->Width, ViewportInfo->Height, bFullscreen, ViewportInfo->PixelFormat );
 		}
 
 		return &ViewportInfo->ViewportRHI;
diff --git a/Engine/Source/Runtime/SlateRHIRenderer/Private/SlateRHIRenderer.h b/Engine/Source/Runtime/SlateRHIRenderer/Private/SlateRHIRenderer.h
--- a/Engine/Source/Runtime/SlateRHIRenderer/Private/SlateRHIRenderer.h
+++ b/Engine/Source/Runtime/SlateRHIRenderer/Private/SlateRHIRenderer.h
@@ -110,6 +110,8 @@ private:
 		bool bRequiresStencilTest;
 		/** Whether or not the viewport is in fullscreen */
 		bool bFullscreen;
+		/** The desired pixel format for this viewport */
+		EPixelFormat PixelFormat;
 	
 		/** FRenderResource interface */
 		virtual void InitRHI() override;
@@ -122,7 +124,8 @@ private:
 				DesiredWidth(0),
 				DesiredHeight(0),
 				bRequiresStencilTest(false),
-				bFullscreen(false)
+				bFullscreen(false),
+				PixelFormat(EPixelFormat::PF_Unknown)
 		{
 
 		}
@@ -171,7 +174,7 @@ public:
 	virtual ISlateAtlasProvider* GetTextureAtlasProvider() override;
 
 	/** Draws windows from a FSlateDrawBuffer on the render thread */
-	void DrawWindow_RenderThread(FRHICommandListImmediate& RHICmdList, const FSlateRHIRenderer::FViewportInfo& ViewportInfo, const FSlateWindowElementList& WindowElementList, bool bLockToVsync);
+	void DrawWindow_RenderThread(FRHICommandListImmediate& RHICmdList, const FSlateRHIRenderer::FViewportInfo& ViewportInfo, const FSlateWindowElementList& WindowElementList, bool bLockToVsync, bool bClear);
 
 
 	/**
diff --git a/Engine/Source/Runtime/SlateRHIRenderer/Private/SlateRHIRenderingPolicy.cpp b/Engine/Source/Runtime/SlateRHIRenderer/Private/SlateRHIRenderingPolicy.cpp
--- a/Engine/Source/Runtime/SlateRHIRenderer/Private/SlateRHIRenderingPolicy.cpp
+++ b/Engine/Source/Runtime/SlateRHIRenderer/Private/SlateRHIRenderingPolicy.cpp
@@ -331,7 +331,7 @@ void FSlateRHIRenderingPolicy::DrawElements(FRHICommandListImmediate& RHICmdList
 				RHICmdList.SetBlendState(
 					(RenderBatch.DrawFlags & ESlateBatchDrawFlag::NoBlending)
 					? TStaticBlendState<>::GetRHI()
-					: TStaticBlendState<CW_RGBA, BO_Add, BF_SourceAlpha, BF_InverseSourceAlpha, BO_Add, BF_One, BF_InverseSourceAlpha>::GetRHI()
+					: TStaticBlendState<CW_RGBA, BO_Add, BF_SourceAlpha, BF_InverseSourceAlpha, BO_Add, BF_InverseDestAlpha, BF_One>::GetRHI()
 					);
 #else
 				RHICmdList.SetBlendState(TStaticBlendState<CW_RGB, BO_Add, BF_One, BF_One, BO_Add, BF_Zero, BF_InverseSourceAlpha>::GetRHI());
