UE4 - Add prev game/real time to View Uniform Buffer in prep for WPO changes

diff --git a/Engine/Source/Runtime/Engine/Public/SceneView.h b/Engine/Source/Runtime/Engine/Public/SceneView.h
--- a/Engine/Source/Runtime/Engine/Public/SceneView.h
+++ b/Engine/Source/Runtime/Engine/Public/SceneView.h
@@ -264,6 +264,8 @@ BEGIN_UNIFORM_BUFFER_STRUCT_WITH_CONSTRUCTOR(FViewUniformShaderParameters,ENGINE
 	DECLARE_UNIFORM_BUFFER_STRUCT_MEMBER_EX(FVector4,SpecularOverrideParameter, EShaderPrecisionModifier::Half)
 	DECLARE_UNIFORM_BUFFER_STRUCT_MEMBER_EX(FVector4,NormalOverrideParameter, EShaderPrecisionModifier::Half)
 	DECLARE_UNIFORM_BUFFER_STRUCT_MEMBER_EX(FVector2D,RoughnessOverrideParameter, EShaderPrecisionModifier::Half)
+	DECLARE_UNIFORM_BUFFER_STRUCT_MEMBER(float, PrevFrameGameTime)
+	DECLARE_UNIFORM_BUFFER_STRUCT_MEMBER(float, PrevFrameRealTime)
 	DECLARE_UNIFORM_BUFFER_STRUCT_MEMBER(FVector,PreViewTranslation)
 	DECLARE_UNIFORM_BUFFER_STRUCT_MEMBER_EX(float,OutOfBoundsMask, EShaderPrecisionModifier::Half)
 	DECLARE_UNIFORM_BUFFER_STRUCT_MEMBER(FVector,ViewOriginDelta)
diff --git a/Engine/Source/Runtime/Renderer/Private/SceneRendering.cpp b/Engine/Source/Runtime/Renderer/Private/SceneRendering.cpp
--- a/Engine/Source/Runtime/Renderer/Private/SceneRendering.cpp
+++ b/Engine/Source/Runtime/Renderer/Private/SceneRendering.cpp
@@ -330,6 +330,8 @@ TUniformBufferRef<FViewUniformShaderParameters> FViewInfo::CreateUniformBuffer(
 	ViewUniformShaderParameters.SpecularOverrideParameter = SpecularOverrideParameter;
 	ViewUniformShaderParameters.NormalOverrideParameter = NormalOverrideParameter;
 	ViewUniformShaderParameters.RoughnessOverrideParameter = RoughnessOverrideParameter;
+	ViewUniformShaderParameters.PrevFrameGameTime = Family->CurrentWorldTime - Family->DeltaWorldTime;
+	ViewUniformShaderParameters.PrevFrameRealTime = Family->CurrentRealTime - Family->DeltaWorldTime;
 	ViewUniformShaderParameters.PreViewTranslation = ViewMatrices.PreViewTranslation;
 	ViewUniformShaderParameters.ViewOriginDelta = ViewMatrices.ViewOrigin - PrevViewMatrices.ViewOrigin;
 	ViewUniformShaderParameters.CullingSign = bReverseCulling ? -1.0f : 1.0f;
diff --git a/Engine/Source/Runtime/SlateRHIRenderer/Private/SlateRHIRenderingPolicy.cpp b/Engine/Source/Runtime/SlateRHIRenderer/Private/SlateRHIRenderingPolicy.cpp
--- a/Engine/Source/Runtime/SlateRHIRenderer/Private/SlateRHIRenderingPolicy.cpp
+++ b/Engine/Source/Runtime/SlateRHIRenderer/Private/SlateRHIRenderingPolicy.cpp
@@ -226,6 +226,8 @@ static FSceneView& CreateSceneView( FSceneViewFamilyContext& ViewFamilyContext,
 	ViewUniformShaderParameters.SpecularOverrideParameter = View->SpecularOverrideParameter;
 	ViewUniformShaderParameters.NormalOverrideParameter = View->NormalOverrideParameter;
 	ViewUniformShaderParameters.RoughnessOverrideParameter = View->RoughnessOverrideParameter;
+	ViewUniformShaderParameters.PrevFrameGameTime = View->Family->CurrentWorldTime - View->Family->DeltaWorldTime;
+	ViewUniformShaderParameters.PrevFrameRealTime = View->Family->CurrentRealTime - View->Family->DeltaWorldTime;
 	ViewUniformShaderParameters.PreViewTranslation = View->ViewMatrices.PreViewTranslation;
 	ViewUniformShaderParameters.CullingSign = View->bReverseCulling ? -1.0f : 1.0f;
 	ViewUniformShaderParameters.NearPlane = GNearClippingPlane;
