Fix for ViewProjectionMatrix not getting set for certain rendering paths

diff --git a/Engine/Source/Runtime/SlateRHIRenderer/Private/SlateRHIRenderingPolicy.cpp b/Engine/Source/Runtime/SlateRHIRenderer/Private/SlateRHIRenderingPolicy.cpp
--- a/Engine/Source/Runtime/SlateRHIRenderer/Private/SlateRHIRenderingPolicy.cpp
+++ b/Engine/Source/Runtime/SlateRHIRenderer/Private/SlateRHIRenderingPolicy.cpp
@@ -449,7 +449,7 @@ void FSlateRHIRenderingPolicy::DrawElements(FRHICommandListImmediate& RHICmdList
 						FGeometryShaderRHIRef()));
 
 					PixelShader->SetParameters(RHICmdList, *SceneView, MaterialRenderProxy, Material, 1.0f / DisplayGamma, ShaderParams.PixelParams);
-
+					VertexShader->SetViewProjection( RHICmdList, ViewProjectionMatrix );
 
 					check(RenderBatch.NumIndices > 0);
 
