No longer suspend the rendering thread when resizing Slate windows.  The suspension was to avoid deadlocks but that looks to be in another way.

diff --git a/Engine/Source/Runtime/SlateRHIRenderer/Private/SlateRHIRenderer.cpp b/Engine/Source/Runtime/SlateRHIRenderer/Private/SlateRHIRenderer.cpp
--- a/Engine/Source/Runtime/SlateRHIRenderer/Private/SlateRHIRenderer.cpp
+++ b/Engine/Source/Runtime/SlateRHIRenderer/Private/SlateRHIRenderer.cpp
@@ -289,9 +289,8 @@ void FSlateRHIRenderer::ConditionalResizeViewport( FViewportInfo* ViewInfo, uint
 		// The viewport size we have doesn't match the requested size of the viewport.
 		// Resize it now.
 
-		// Suspend the rendering thread to avoid deadlocks with the gpu
-		bool bRecreateThread = true;
-		SCOPED_SUSPEND_RENDERING_THREAD( bRecreateThread );
+		// cannot resize the viewport while potentially using it.
+		FlushRenderingCommands();
 
 		// Windows are allowed to be zero sized ( sometimes they are animating to/from zero for example)
 		// but viewports cannot be zero sized.  Use 8x8 as a reasonably sized viewport in this case.
