Missed a change from last submission

diff --git a/Engine/Source/Runtime/SlateRHIRenderer/Private/SlateRHIRenderer.cpp b/Engine/Source/Runtime/SlateRHIRenderer/Private/SlateRHIRenderer.cpp
--- a/Engine/Source/Runtime/SlateRHIRenderer/Private/SlateRHIRenderer.cpp
+++ b/Engine/Source/Runtime/SlateRHIRenderer/Private/SlateRHIRenderer.cpp
@@ -325,9 +325,6 @@ void FSlateRHIRenderer::ConditionalResizeViewport( FViewportInfo* ViewInfo, uint
 		{
 			ViewInfo->ViewportRHI = RHICreateViewport(ViewInfo->OSWindow, NewWidth, NewHeight, bFullscreen, ViewInfo->PixelFormat);
 		}
-
-		// Safe to call here as the rendering thread has been suspended: game thread == render thread!
-		ViewInfo->RecreateDepthBuffer_RenderThread();
 	}
 }
 
