Integrated: GitHub 847 : Aligned thread check in GetMaterialResource() with rest of code

"The check for being called from the correct thread in FSlateRHIResourceManager::GetMaterialResource() (IsInGameThread) was not in line with the rest of the FSlateRHIResourceManager (IsThreadSafeForSlateRendering).
This triggered an assertion when using materials from within a slate widget that is used by the movie player."

diff --git a/Engine/Source/Runtime/SlateRHIRenderer/Private/SlateRHIResourceManager.cpp b/Engine/Source/Runtime/SlateRHIRenderer/Private/SlateRHIResourceManager.cpp
--- a/Engine/Source/Runtime/SlateRHIRenderer/Private/SlateRHIResourceManager.cpp
+++ b/Engine/Source/Runtime/SlateRHIRenderer/Private/SlateRHIResourceManager.cpp
@@ -594,7 +594,7 @@ FSlateShaderResourceProxy* FSlateRHIResourceManager::FindOrCreateDynamicTextureR
 
 FSlateShaderResourceProxy* FSlateRHIResourceManager::GetMaterialResource(const FSlateBrush& InBrush)
 {
-	check(IsInGameThread());
+	check(IsThreadSafeForSlateRendering());
 
 	const FName ResourceName = InBrush.GetResourceName();
 
