No longer report utexture references from the Slate resource manager.  This fixes UTextures being permenently refererenced by Slate/UMG even when not in use

diff --git a/Engine/Source/Runtime/SlateRHIRenderer/Private/SlateRHIResourceManager.cpp b/Engine/Source/Runtime/SlateRHIRenderer/Private/SlateRHIResourceManager.cpp
--- a/Engine/Source/Runtime/SlateRHIRenderer/Private/SlateRHIResourceManager.cpp
+++ b/Engine/Source/Runtime/SlateRHIRenderer/Private/SlateRHIResourceManager.cpp
@@ -100,25 +100,13 @@ void FDynamicResourceMap::ReleaseResources()
 		BeginReleaseResource(It.Value()->RHIRefTexture);
 	}
 	
-	for (TMap<UObject*, TSharedPtr<FSlateUTextureResource> >::TIterator It(UTextureResourceMap); It; ++It)
+	for (TMap<TWeakObjectPtr<UTexture2D>, TSharedPtr<FSlateUTextureResource> >::TIterator It(UTextureResourceMap); It; ++It)
 	{
 		It.Value()->UpdateRenderResource(nullptr);
 	}
 
 }
 
-void FDynamicResourceMap::AddReferencedObjects(FReferenceCollector& Collector)
-{
-	for(TMap<UObject*, TSharedPtr<FSlateUTextureResource> >::TIterator It(UTextureResourceMap); It; ++It)
-	{
-		TSharedPtr<FSlateUTextureResource>& Resource = It.Value();
-		if (Resource.IsValid() && Resource->TextureObject != nullptr)
-		{
-			Collector.AddReferencedObject(Resource->TextureObject);
-		}
-	}
-}
-
 
 FSlateRHIResourceManager::FSlateRHIResourceManager()
 {
@@ -369,9 +357,16 @@ static void LoadUObjectForBrush( const FSlateBrush& InBrush )
 		// Set the texture object to a default texture to prevent constant loading of missing textures
 		if( !TextureObject )
 		{
+			UE_LOG(LogSlate, Warning, TEXT("Error loading loading UTexture from path: %s not found"), *Path);
 			TextureObject = GEngine->DefaultTexture;
 		}
+		else
+		{
+			// We do this here because this deprecated system of loading textures will not report references and we dont want the Slate RHI resource manager to manage references
+			TextureObject->AddToRoot();
+		}
 
+		
 		Brush->SetResourceObject(TextureObject);
 
 		UE_LOG(LogSlate, Warning, TEXT("The texture:// method of loading UTextures for use in Slate is deprecated.  Please convert %s to a Brush Asset"), *Path);
diff --git a/Engine/Source/Runtime/SlateRHIRenderer/Private/SlateRHIResourceManager.h b/Engine/Source/Runtime/SlateRHIRenderer/Private/SlateRHIResourceManager.h
--- a/Engine/Source/Runtime/SlateRHIRenderer/Private/SlateRHIResourceManager.h
+++ b/Engine/Source/Runtime/SlateRHIRenderer/Private/SlateRHIResourceManager.h
@@ -8,7 +8,7 @@ class FSlateDynamicTextureResource;
 class FSlateUTextureResource;
 class FSlateMaterialResource;
 
-struct FDynamicResourceMap : public FGCObject
+struct FDynamicResourceMap
 {
 public:
 	TSharedPtr<FSlateDynamicTextureResource> GetDynamicTextureResource( FName ResourceName ) const;
@@ -35,14 +35,10 @@ public:
 	void ReleaseResources();
 
 	uint32 GetNumObjectResources() const { return UTextureResourceMap.Num() + MaterialResourceMap.Num(); }
-
-	/** FGCObject interface */
-	virtual void AddReferencedObjects(FReferenceCollector& Collector) override;
-
 private:
 	TMap<FName, TSharedPtr<FSlateDynamicTextureResource> > NativeTextureMap;
 	
-	TMap<UObject*, TSharedPtr<FSlateUTextureResource> > UTextureResourceMap;
+	TMap<TWeakObjectPtr<UTexture2D>, TSharedPtr<FSlateUTextureResource> > UTextureResourceMap;
 
 	/** Map of all material resources */
 	TMap<TWeakObjectPtr<UMaterialInterface>, TSharedPtr<FSlateMaterialResource> > MaterialResourceMap;
