Canvas runtime font quality improvements when using a scaled font glyph

We now use bilinear sampling rather than point sampling.

diff --git a/Engine/Source/Runtime/SlateRHIRenderer/Private/SlateRHIFontTexture.cpp b/Engine/Source/Runtime/SlateRHIRenderer/Private/SlateRHIFontTexture.cpp
--- a/Engine/Source/Runtime/SlateRHIRenderer/Private/SlateRHIFontTexture.cpp
+++ b/Engine/Source/Runtime/SlateRHIRenderer/Private/SlateRHIFontTexture.cpp
@@ -23,33 +23,31 @@ void FSlateFontTextureRHI::InitDynamicRHI()
 		// Also assign the reference to the FTextureResource variable so that the Engine can access it
 		TextureRHI = ShaderResource;
 
-		const float MipMapBias = -1.0f;
-
 		// Create the sampler state RHI resource.
 		FSamplerStateInitializerRHI SamplerStateInitializer
 		(
-		  SF_Point,
+		  SF_Bilinear,
 		  AM_Clamp,
 		  AM_Clamp,
 		  AM_Wrap,
-		  MipMapBias
+		  0,
+		  1, // Disable anisotropic filtering, since aniso doesn't respect MaxLOD
+		  0,
+		  0
 		);
 		SamplerStateRHI = RHICreateSamplerState(SamplerStateInitializer);
 
 		// Create a custom sampler state for using this texture in a deferred pass, where ddx / ddy are discontinuous
 		FSamplerStateInitializerRHI DeferredPassSamplerStateInitializer
 		(
-		  SF_Point,
+		  SF_Bilinear,
 		  AM_Clamp,
 		  AM_Clamp,
 		  AM_Wrap,
-		  MipMapBias,
-		  // Disable anisotropic filtering, since aniso doesn't respect MaxLOD
-		  1,
 		  0,
-		  // Prevent the less detailed mip levels from being used, which hides artifacts on silhouettes due to ddx / ddy being very large
-		  // This has the side effect that it increases minification aliasing on light functions
-		  2
+		  1, // Disable anisotropic filtering, since aniso doesn't respect MaxLOD
+		  0,
+		  0
 		);
 		DeferredPassSamplerStateRHI = RHICreateSamplerState(DeferredPassSamplerStateInitializer);
 
