Fixed gamma not being accurate in Slate.  We now use the same gamma correction scheme as in the tone mapper which uses the more accurate srgb standard.  (e.g Slate UI should look the same as in photoshop assuming it was imported correctly).

- Consolidated gamma correction functions into their own shader file so they can be shared.

diff --git a/Engine/Shaders/GammaCorrectionCommon.usf b/Engine/Shaders/GammaCorrectionCommon.usf
new file mode 100644
--- /dev/null
+++ b/Engine/Shaders/GammaCorrectionCommon.usf
@@ -0,0 +1,30 @@
+// Copyright 1998-2015 Epic Games, Inc. All Rights Reserved.
+
+
+half3 LinearTo709Branchless(half3 lin) 
+{
+	lin = max(6.10352e-5, lin); // minimum positive non-denormal (fixes black problem on DX11 AMD and NV)
+	return min(lin * 4.5, pow(max(lin, 0.018), 0.45) * 1.099 - 0.099);
+}
+
+half3 LinearToSrgbBranchless(half3 lin) 
+{
+	lin = max(6.10352e-5, lin); // minimum positive non-denormal (fixes black problem on DX11 AMD and NV)
+	return min(lin * 12.92, pow(max(lin, 0.00313067), 1.0/2.4) * 1.055 - 0.055);
+	// Possible that mobile GPUs might have native pow() function?
+	//return min(lin * 12.92, exp2(log2(max(lin, 0.00313067)) * (1.0/2.4) + log2(1.055)) - 0.055);
+}
+
+half LinearToSrgbBranchingChannel(half lin) 
+{
+	if(lin < 0.00313067) return lin * 12.92;
+	return pow(lin, (1.0/2.4)) * 1.055 - 0.055;
+}
+
+half3 LinearToSrgbBranching(half3 lin) 
+{
+	return half3(
+		LinearToSrgbBranchingChannel(lin.r),
+		LinearToSrgbBranchingChannel(lin.g),
+		LinearToSrgbBranchingChannel(lin.b));
+}
diff --git a/Engine/Shaders/PostProcessTonemap.usf b/Engine/Shaders/PostProcessTonemap.usf
--- a/Engine/Shaders/PostProcessTonemap.usf
+++ b/Engine/Shaders/PostProcessTonemap.usf
@@ -7,7 +7,7 @@
 #include "Common.usf"
 #include "PostProcessCommon.usf"
 #include "TonemapCommon.usf"
-
+#include "GammaCorrectionCommon.usf"
 //
 // 64x64 tiling noise texture, optimized to have mostly high frequency content
 Texture2D NoiseTexture;
@@ -36,34 +36,6 @@ half GrainFromUV(float2 GrainUV)
 	return Grain;
 }
 
-half3 LinearTo709Branchless(half3 lin) 
-{
-	lin = max(6.10352e-5, lin); // minimum positive non-denormal (fixes black problem on DX11 AMD and NV)
-	return min(lin * 4.5, pow(max(lin, 0.018), 0.45) * 1.099 - 0.099);
-}
-
-half3 LinearToSrgbBranchless(half3 lin) 
-{
-	lin = max(6.10352e-5, lin); // minimum positive non-denormal (fixes black problem on DX11 AMD and NV)
-	return min(lin * 12.92, pow(max(lin, 0.00313067), 1.0/2.4) * 1.055 - 0.055);
-	// Possible that mobile GPUs might have native pow() function?
-	//return min(lin * 12.92, exp2(log2(max(lin, 0.00313067)) * (1.0/2.4) + log2(1.055)) - 0.055);
-}
-
-half LinearToSrgbBranchingChannel(half lin) 
-{
-	if(lin < 0.00313067) return lin * 12.92;
-	return pow(lin, (1.0/2.4)) * 1.055 - 0.055;
-}
-
-half3 LinearToSrgbBranching(half3 lin) 
-{
-	return half3(
-		LinearToSrgbBranchingChannel(lin.r),
-		LinearToSrgbBranchingChannel(lin.g),
-		LinearToSrgbBranchingChannel(lin.b));
-}
-
 
 
 // 256x16 color LUT for color grading
diff --git a/Engine/Shaders/SlateElementPixelShader.usf b/Engine/Shaders/SlateElementPixelShader.usf
--- a/Engine/Shaders/SlateElementPixelShader.usf
+++ b/Engine/Shaders/SlateElementPixelShader.usf
@@ -2,6 +2,7 @@
 
 
 #include "Common.usf"
+#include "GammaCorrectionCommon.usf"
 
 #if USE_MATERIALS
 #include "Material.usf"
@@ -26,8 +27,8 @@
 /** Enable to debug clipping issues */
 #define DEBUG_CLIPPING 0
 
-/** Display gamma */
-float InvDisplayGamma;
+/** Display gamma x:gamma curve adjustment, y:inverse gamma (1/GEngine->DisplayGamma) */
+float2 GammaValues;
 /** Per element params */
 float4 ShaderParams;
 
@@ -38,6 +39,36 @@ SamplerState ElementTextureSampler;
 
 #endif
 
+
+float3 GammaCorrect(float3 InColor)
+{
+	float3 CorrectedColor = InColor;
+
+#if !(ES2_PROFILE || ES3_1_PROFILE)
+	FLATTEN if( GammaValues.y != 1.0f )
+	{
+		// Apply "gamma" curve adjustment.
+		CorrectedColor = pow(CorrectedColor, GammaValues.x);
+
+		#if MAC
+			// Note, MacOSX native output is raw gamma 2.2 not sRGB!
+			CorrectedColor = pow(CorrectedColor, 1.0/2.2);
+		#else
+			#if USE_709
+				// Didn't profile yet if the branching version would be faster (different linear segment).
+				CorrectedColor = LinearTo709Branchless(CorrectedColor);
+			#else
+				CorrectedColor = LinearToSrgbBranching(CorrectedColor);
+			#endif
+		#endif
+	}
+#endif
+
+	return CorrectedColor;
+}
+
+
+
 struct VertexOut
 {
 	float4 Position : SV_POSITION;
@@ -68,8 +99,19 @@ float4 GetMaterialColor( VertexOut InVertex, float2 UV )
 
 	// Clip if masked
 	GetMaterialCoverageAndClipping( Parameters );
+	
+	float Opacity = GetMaterialOpacity(Parameters);
+	float3 Color = GetMaterialEmissive(Parameters);
+
+	float4 OutColor;
 
-	return float4( GetMaterialEmissive( Parameters ), GetMaterialOpacity(Parameters) );
+	#if MATERIALBLENDING_ADDITIVE
+		OutColor = float4(Color * Opacity, 0.0f);
+	#else
+		OutColor = float4( Color, Opacity );
+	#endif
+
+	return OutColor;
 }
 #endif
 
@@ -327,10 +369,8 @@ float4 Main( VertexOut VIn ) : SV_Target0
 	OutColor = ApplyColorVisionDeficiencyFilter(OutColor);
 #endif
 
-#if !(ES2_PROFILE || ES3_1_PROFILE)
 	// gamma correct
-	OutColor.rgb = pow(OutColor.rgb,InvDisplayGamma);
-#endif
+	OutColor.rgb = GammaCorrect(OutColor.rgb);
 
 
 #if DEBUG_CLIPPING
diff --git a/Engine/Source/Runtime/Engine/Private/ShaderDerivedDataVersion.h b/Engine/Source/Runtime/Engine/Private/ShaderDerivedDataVersion.h
--- a/Engine/Source/Runtime/Engine/Private/ShaderDerivedDataVersion.h
+++ b/Engine/Source/Runtime/Engine/Private/ShaderDerivedDataVersion.h
@@ -10,4 +10,4 @@
 // guid as version
 
 #define GLOBALSHADERMAP_DERIVEDDATA_VER			TEXT("af1606fcec214373844d2b6cf9df71c8")
-#define MATERIALSHADERMAP_DERIVEDDATA_VER		TEXT("bfbd0e1a93dd4482a991041cedfae4e9")
+#define MATERIALSHADERMAP_DERIVEDDATA_VER		TEXT("983afa3849534faa8491540b994bd32e")
diff --git a/Engine/Source/Runtime/SlateRHIRenderer/Private/SlateMaterialShader.cpp b/Engine/Source/Runtime/SlateRHIRenderer/Private/SlateMaterialShader.cpp
--- a/Engine/Source/Runtime/SlateRHIRenderer/Private/SlateMaterialShader.cpp
+++ b/Engine/Source/Runtime/SlateRHIRenderer/Private/SlateMaterialShader.cpp
@@ -22,7 +22,7 @@ FSlateMaterialShaderPS::FSlateMaterialShaderPS(const FMaterialShaderType::Compil
 	: FMaterialShader(Initializer)
 {
 	ShaderParams.Bind(Initializer.ParameterMap, TEXT("ShaderParams"));
-	DisplayGamma.Bind(Initializer.ParameterMap, TEXT("InvDisplayGamma"));
+	GammaValues.Bind(Initializer.ParameterMap, TEXT("GammaValues"));
 }
 
 void FSlateMaterialShaderPS::SetParameters(FRHICommandList& RHICmdList, const FSceneView& View, const FMaterialRenderProxy* MaterialRenderProxy, const FMaterial* Material, float InDisplayGamma, const FVector4& InShaderParams )
@@ -53,7 +53,6 @@ void FSlateMaterialShaderPS::SetParameters(FRHICommandList& RHICmdList, const FS
 		break;
 	};
 
-	SetShaderValue( RHICmdList, ShaderRHI, DisplayGamma, InDisplayGamma );
 	SetShaderValue( RHICmdList, ShaderRHI, ShaderParams, InShaderParams );
 
 	const bool bDeferredPass = false;
@@ -62,10 +61,17 @@ void FSlateMaterialShaderPS::SetParameters(FRHICommandList& RHICmdList, const FS
 }
 
 
+void FSlateMaterialShaderPS::SetDisplayGamma(FRHICommandList& RHICmdList, float InDisplayGamma)
+{
+	FVector2D InGammaValues(2.2f / InDisplayGamma, 1.0f/InDisplayGamma);
+
+	SetShaderValue(RHICmdList, GetPixelShader(), GammaValues, InGammaValues);
+}
+
 bool FSlateMaterialShaderPS::Serialize(FArchive& Ar)
 {
 	bool bShaderHasOutdatedParameters = FMaterialShader::Serialize(Ar);
-	Ar << DisplayGamma;
+	Ar << GammaValues;
 	Ar << ShaderParams;
 	return bShaderHasOutdatedParameters;
 }
diff --git a/Engine/Source/Runtime/SlateRHIRenderer/Private/SlateMaterialShader.h b/Engine/Source/Runtime/SlateRHIRenderer/Private/SlateMaterialShader.h
--- a/Engine/Source/Runtime/SlateRHIRenderer/Private/SlateMaterialShader.h
+++ b/Engine/Source/Runtime/SlateRHIRenderer/Private/SlateMaterialShader.h
@@ -20,11 +20,13 @@ public:
 
 	void SetParameters(FRHICommandList& RHICmdList, const FSceneView& View, const FMaterialRenderProxy* MaterialRenderProxy, const FMaterial* Material, float InDisplayGamma, const FVector4& InShaderParams );
 
+	void SetDisplayGamma(FRHICommandList& RHICmdList, float InDisplayGamma);
+
 	virtual bool Serialize(FArchive& Ar) override;
 
 private:
 
-	FShaderParameter DisplayGamma;
+	FShaderParameter GammaValues;
 	FShaderParameter ShaderParams;
 };
 
diff --git a/Engine/Source/Runtime/SlateRHIRenderer/Private/SlateRHIRenderingPolicy.cpp b/Engine/Source/Runtime/SlateRHIRenderer/Private/SlateRHIRenderingPolicy.cpp
--- a/Engine/Source/Runtime/SlateRHIRenderer/Private/SlateRHIRenderingPolicy.cpp
+++ b/Engine/Source/Runtime/SlateRHIRenderer/Private/SlateRHIRenderingPolicy.cpp
@@ -407,7 +407,7 @@ void FSlateRHIRenderingPolicy::DrawElements(FRHICommandListImmediate& RHICmdList
 				}
 
 				PixelShader->SetShaderParams(RHICmdList, ShaderParams.PixelParams);
-				PixelShader->SetDisplayGamma(RHICmdList, (DrawFlags & ESlateBatchDrawFlag::NoGamma) ? 1.0f : 1.0f/DisplayGamma);
+				PixelShader->SetDisplayGamma(RHICmdList, (DrawFlags & ESlateBatchDrawFlag::NoGamma) ? 1.0f : DisplayGamma);
 
 
 				check(RenderBatch.NumIndices > 0);
@@ -449,7 +449,9 @@ void FSlateRHIRenderingPolicy::DrawElements(FRHICommandListImmediate& RHICmdList
 						PixelShader->GetPixelShader(),
 						FGeometryShaderRHIRef()));
 
-					PixelShader->SetParameters(RHICmdList, *SceneView, MaterialRenderProxy, Material, 1.0f / DisplayGamma, ShaderParams.PixelParams);
+					PixelShader->SetParameters(RHICmdList, *SceneView, MaterialRenderProxy, Material, DisplayGamma, ShaderParams.PixelParams);
+					PixelShader->SetDisplayGamma(RHICmdList, (DrawFlags & ESlateBatchDrawFlag::NoGamma) ? 1.0f : DisplayGamma);
+
 					VertexShader->SetViewProjection( RHICmdList, ViewProjectionMatrix );
 
 					check(RenderBatch.NumIndices > 0);
diff --git a/Engine/Source/Runtime/SlateRHIRenderer/Private/SlateShaders.cpp b/Engine/Source/Runtime/SlateRHIRenderer/Private/SlateShaders.cpp
--- a/Engine/Source/Runtime/SlateRHIRenderer/Private/SlateShaders.cpp
+++ b/Engine/Source/Runtime/SlateRHIRenderer/Private/SlateShaders.cpp
@@ -63,6 +63,12 @@ void FSlateVertexDeclaration::ReleaseRHI()
 	VertexDeclarationRHI.SafeRelease();
 }
 
+void FSlateElementPS::ModifyCompilationEnvironment(EShaderPlatform Platform, FShaderCompilerEnvironment& OutEnvironment)
+{
+	static const auto CVar = IConsoleManager::Get().FindTConsoleVariableDataInt(TEXT("r.Tonemapper709"));
+	OutEnvironment.SetDefine(TEXT("USE_709"), CVar ? (CVar->GetValueOnGameThread() != 0) : 1);
+}
+
 /************************************************************************/
 /* FSlateDefaultVertexShader                                            */
 /************************************************************************/
diff --git a/Engine/Source/Runtime/SlateRHIRenderer/Private/SlateShaders.h b/Engine/Source/Runtime/SlateRHIRenderer/Private/SlateShaders.h
--- a/Engine/Source/Runtime/SlateRHIRenderer/Private/SlateShaders.h
+++ b/Engine/Source/Runtime/SlateRHIRenderer/Private/SlateShaders.h
@@ -88,9 +88,10 @@ public:
 		TextureParameter.Bind( Initializer.ParameterMap, TEXT("ElementTexture"));
 		TextureParameterSampler.Bind( Initializer.ParameterMap, TEXT("ElementTextureSampler"));
 		ShaderParams.Bind( Initializer.ParameterMap, TEXT("ShaderParams"));
-		DisplayGamma.Bind( Initializer.ParameterMap,TEXT("InvDisplayGamma"));
+		GammaValues.Bind( Initializer.ParameterMap,TEXT("GammaValues"));
 	}
 
+	static void ModifyCompilationEnvironment(EShaderPlatform Platform, FShaderCompilerEnvironment& OutEnvironment);
 
 	/**
 	 * Sets the texture used by this shader 
@@ -120,7 +121,9 @@ public:
 	 */
 	void SetDisplayGamma(FRHICommandList& RHICmdList, float InDisplayGamma)
 	{
-		SetShaderValue(RHICmdList, GetPixelShader(),DisplayGamma,InDisplayGamma);
+		FVector2D InGammaValues( 2.2f / InDisplayGamma, 1.0f/InDisplayGamma );
+
+		SetShaderValue(RHICmdList, GetPixelShader(),GammaValues,InGammaValues);
 	}
 
 	virtual bool Serialize( FArchive& Ar )
@@ -130,7 +133,7 @@ public:
 		Ar << TextureParameter;
 		Ar << TextureParameterSampler;
 		Ar << ShaderParams;
-		Ar << DisplayGamma;
+		Ar << GammaValues;
 
 		return bShaderHasOutdatedParameters;
 	}
@@ -139,7 +142,7 @@ private:
 	FShaderResourceParameter TextureParameter;
 	FShaderResourceParameter TextureParameterSampler;
 	FShaderParameter ShaderParams;
-	FShaderParameter DisplayGamma;
+	FShaderParameter GammaValues;
 };
 
 /** 
@@ -168,11 +171,11 @@ public:
 	static void ModifyCompilationEnvironment(EShaderPlatform Platform, FShaderCompilerEnvironment& OutEnvironment)
 	{
 		// Set defines based on what this shader will be used for
-		OutEnvironment.SetDefine( TEXT("SHADER_TYPE"), (uint32)ShaderType );
-		OutEnvironment.SetDefine( TEXT("DRAW_DISABLED_EFFECT"), (uint32)(bDrawDisabledEffect ? 1 : 0) );
-		OutEnvironment.SetDefine( TEXT("USE_TEXTURE_ALPHA"), (uint32)(bUseTextureAlpha ? 1 : 0) );
-		OutEnvironment.SetDefine( TEXT("COLOR_VISION_DEFICIENCY_TYPE"), GSlateShaderColorVisionDeficiencyType );
-		OutEnvironment.SetDefine( TEXT("USE_MATERIALS"), (uint32)0 );
+		OutEnvironment.SetDefine(TEXT("SHADER_TYPE"), (uint32)ShaderType);
+		OutEnvironment.SetDefine(TEXT("DRAW_DISABLED_EFFECT"), (uint32)(bDrawDisabledEffect ? 1 : 0));
+		OutEnvironment.SetDefine(TEXT("USE_TEXTURE_ALPHA"), (uint32)(bUseTextureAlpha ? 1 : 0));
+		OutEnvironment.SetDefine(TEXT("COLOR_VISION_DEFICIENCY_TYPE"), GSlateShaderColorVisionDeficiencyType);
+		OutEnvironment.SetDefine(TEXT("USE_MATERIALS"), (uint32)0);
 
 		FSlateElementPS::ModifyCompilationEnvironment( Platform, OutEnvironment );
 	}
