Fix android 3d widgets rendering upside down

diff --git a/Engine/Shaders/SlateVertexShader.usf b/Engine/Shaders/SlateVertexShader.usf
--- a/Engine/Shaders/SlateVertexShader.usf
+++ b/Engine/Shaders/SlateVertexShader.usf
@@ -5,6 +5,8 @@
 
 float4x4 ViewProjection;
 
+float SwitchVerticalAxisMultiplier;
+
 struct VertexOut
 {
 	float4 Position : SV_POSITION;
@@ -36,7 +38,7 @@ VertexOut Main(
 
 #if ES2_PROFILE && COMPILER_GLSL_ES2
 	// @todo-mobile: Fix this in the projection matrix
-	VOut.Position.y =  -VOut.Position.y;
+	VOut.Position.y *= SwitchVerticalAxisMultiplier;
 #endif
 
 
diff --git a/Engine/Source/Runtime/SlateRHIRenderer/Private/Slate3DRenderer.cpp b/Engine/Source/Runtime/SlateRHIRenderer/Private/Slate3DRenderer.cpp
--- a/Engine/Source/Runtime/SlateRHIRenderer/Private/Slate3DRenderer.cpp
+++ b/Engine/Source/Runtime/SlateRHIRenderer/Private/Slate3DRenderer.cpp
@@ -104,12 +104,16 @@ void FSlate3DRenderer::DrawWindowToTarget_RenderThread( FRHICommandListImmediate
 	{
 		FSlateBackBuffer BackBufferTarget(RenderTargetResource->GetTextureRHI(), FIntPoint(RenderTarget->SizeX, RenderTarget->SizeY ) );
 
+		// The scene renderer will handle it in this case
+		const bool bAllowSwitchVerticalAxis = false;
+
 		RenderTargetPolicy->DrawElements
 		(
 			RHICmdList,
 			BackBufferTarget,
 			ProjectionMatrix,
-			WindowElementList.GetRenderBatches()
+			WindowElementList.GetRenderBatches(),
+			bAllowSwitchVerticalAxis
 			);
 	}
 
diff --git a/Engine/Source/Runtime/SlateRHIRenderer/Private/SlateRHIRenderingPolicy.cpp b/Engine/Source/Runtime/SlateRHIRenderer/Private/SlateRHIRenderingPolicy.cpp
--- a/Engine/Source/Runtime/SlateRHIRenderer/Private/SlateRHIRenderingPolicy.cpp
+++ b/Engine/Source/Runtime/SlateRHIRenderer/Private/SlateRHIRenderingPolicy.cpp
@@ -264,7 +264,7 @@ static FSceneView& CreateSceneView( FSceneViewFamilyContext& ViewFamilyContext,
 	return *View;
 }
 
-void FSlateRHIRenderingPolicy::DrawElements(FRHICommandListImmediate& RHICmdList, FSlateBackBuffer& BackBuffer, const FMatrix& ViewProjectionMatrix, const TArray<FSlateRenderBatch>& RenderBatches)
+void FSlateRHIRenderingPolicy::DrawElements(FRHICommandListImmediate& RHICmdList, FSlateBackBuffer& BackBuffer, const FMatrix& ViewProjectionMatrix, const TArray<FSlateRenderBatch>& RenderBatches, bool bAllowSwitchVerticalAxis)
 {
 	SCOPE_CYCLE_COUNTER( STAT_SlateDrawTime );
 
@@ -335,7 +335,7 @@ void FSlateRHIRenderingPolicy::DrawElements(FRHICommandListImmediate& RHICmdList
 
 
 				VertexShader->SetViewProjection(RHICmdList, ViewProjectionMatrix);
-
+				VertexShader->SetVerticalAxisMultiplier(RHICmdList, bAllowSwitchVerticalAxis && RHINeedsToSwitchVerticalAxis(GShaderPlatformForFeatureLevel[GMaxRHIFeatureLevel]) ? -1.0f : 1.0f );
 #if !DEBUG_OVERDRAW
 				RHICmdList.SetBlendState(
 					(RenderBatch.DrawFlags & ESlateBatchDrawFlag::NoBlending)
diff --git a/Engine/Source/Runtime/SlateRHIRenderer/Private/SlateRHIRenderingPolicy.h b/Engine/Source/Runtime/SlateRHIRenderer/Private/SlateRHIRenderingPolicy.h
--- a/Engine/Source/Runtime/SlateRHIRenderer/Private/SlateRHIRenderingPolicy.h
+++ b/Engine/Source/Runtime/SlateRHIRenderer/Private/SlateRHIRenderingPolicy.h
@@ -177,7 +177,7 @@ public:
 	~FSlateRHIRenderingPolicy();
 
 	virtual void UpdateBuffers( const FSlateWindowElementList& WindowElementList ) override;
-	virtual void DrawElements(FRHICommandListImmediate& RHICmdList, class FSlateBackBuffer& BackBuffer, const FMatrix& ViewProjectionMatrix, const TArray<FSlateRenderBatch>& RenderBatches);
+	virtual void DrawElements(FRHICommandListImmediate& RHICmdList, class FSlateBackBuffer& BackBuffer, const FMatrix& ViewProjectionMatrix, const TArray<FSlateRenderBatch>& RenderBatches, bool bAllowSwtichVerticalAxis=true);
 
 	virtual TSharedRef<FSlateFontCache> GetFontCache() override { return FontCache.ToSharedRef(); }
 	virtual TSharedRef<FSlateShaderResourceManager> GetResourceManager() override { return ResourceManager; }
diff --git a/Engine/Source/Runtime/SlateRHIRenderer/Private/SlateShaders.cpp b/Engine/Source/Runtime/SlateRHIRenderer/Private/SlateShaders.cpp
--- a/Engine/Source/Runtime/SlateRHIRenderer/Private/SlateShaders.cpp
+++ b/Engine/Source/Runtime/SlateRHIRenderer/Private/SlateShaders.cpp
@@ -78,6 +78,7 @@ FSlateElementVS::FSlateElementVS( const ShaderMetaType::CompiledShaderInitialize
 {
 	ViewProjection.Bind(Initializer.ParameterMap, TEXT("ViewProjection"));
 	VertexShaderParams.Bind( Initializer.ParameterMap, TEXT("VertexShaderParams"));
+	SwitchVerticalAxisMultiplier.Bind( Initializer.ParameterMap, TEXT("SwitchVerticalAxisMultiplier"));
 }
 
 void FSlateElementVS::SetViewProjection(FRHICommandList& RHICmdList, const FMatrix& InViewProjection )
@@ -90,6 +91,11 @@ void FSlateElementVS::SetShaderParameters(FRHICommandList& RHICmdList, const FVe
 	SetShaderValue(RHICmdList, GetVertexShader(), VertexShaderParams, ShaderParams );
 }
 
+void FSlateElementVS::SetVerticalAxisMultiplier(FRHICommandList& RHICmdList, float InMultiplier )
+{
+	SetShaderValue(RHICmdList, GetVertexShader(), SwitchVerticalAxisMultiplier, InMultiplier );
+}
+
 /** Serializes the shader data */
 bool FSlateElementVS::Serialize( FArchive& Ar )
 {
@@ -97,6 +103,7 @@ bool FSlateElementVS::Serialize( FArchive& Ar )
 
 	Ar << ViewProjection;
 	Ar << VertexShaderParams;
+	Ar << SwitchVerticalAxisMultiplier;
 
 	return bShaderHasOutdatedParameters;
 }
diff --git a/Engine/Source/Runtime/SlateRHIRenderer/Private/SlateShaders.h b/Engine/Source/Runtime/SlateRHIRenderer/Private/SlateShaders.h
--- a/Engine/Source/Runtime/SlateRHIRenderer/Private/SlateShaders.h
+++ b/Engine/Source/Runtime/SlateRHIRenderer/Private/SlateShaders.h
@@ -55,6 +55,11 @@ public:
 	 */
 	void SetShaderParameters(FRHICommandList& RHICmdList, const FVector4& ShaderParams );
 
+	/**
+	 * Sets the vertical axis multiplier to use depending on graphics api
+	 */
+	void SetVerticalAxisMultiplier(FRHICommandList& RHICmdList, float InMultiplier);
+
 	/** Serializes the shader data */
 	virtual bool Serialize( FArchive& Ar );
 
@@ -63,6 +68,8 @@ private:
 	FShaderParameter ViewProjection;
 	/** Shader parmeters used by the shader */
 	FShaderParameter VertexShaderParams;
+	/** Parameter used to determine if we need to swtich the vertical axis for opengl */
+	FShaderParameter SwitchVerticalAxisMultiplier;
 };
 
 /** 
