Prevent materials from being used in early loading screens before the engine is initialized.

diff --git a/Engine/Source/Runtime/SlateRHIRenderer/Private/SlateRHIRenderingPolicy.cpp b/Engine/Source/Runtime/SlateRHIRenderer/Private/SlateRHIRenderingPolicy.cpp
--- a/Engine/Source/Runtime/SlateRHIRenderer/Private/SlateRHIRenderingPolicy.cpp
+++ b/Engine/Source/Runtime/SlateRHIRenderer/Private/SlateRHIRenderingPolicy.cpp
@@ -421,8 +421,9 @@ void FSlateRHIRenderingPolicy::DrawElements(FRHICommandListImmediate& RHICmdList
 				}
 
 			}
-			else if (ShaderResource && ShaderResource->GetType() == ESlateShaderResource::Material)
+			else if (ShaderResource && ShaderResource->GetType() == ESlateShaderResource::Material && GEngine)
 			{
+				// Note: This code is only executed if the engine is loaded (in early loading screens attemping to use a material is unsupported
 				if (!SceneView)
 				{
 					SceneView = &CreateSceneView(SceneViewContext, BackBuffer, ViewProjectionMatrix);
