Ensure flushing of rendering commands by slate only happens on the game thread.  It is poinless for slate to try to flush when it is already on the render thread

diff --git a/Engine/Source/Runtime/SlateRHIRenderer/Private/SlateRHIRenderer.cpp b/Engine/Source/Runtime/SlateRHIRenderer/Private/SlateRHIRenderer.cpp
--- a/Engine/Source/Runtime/SlateRHIRenderer/Private/SlateRHIRenderer.cpp
+++ b/Engine/Source/Runtime/SlateRHIRenderer/Private/SlateRHIRenderer.cpp
@@ -1068,8 +1068,10 @@ bool FSlateRHIRenderer::GenerateDynamicImageResource( FName ResourceName, uint32
  */
 void FSlateRHIRenderer::FlushCommands() const
 {
-	check(!IsInSlateThread());
-	FlushRenderingCommands();
+	if( IsInGameThread() )
+	{
+		FlushRenderingCommands();
+	}
 }
 
 /**
