Fixed copy paste error that had thankfully never been in use.

diff --git a/Engine/Source/Runtime/SlateRHIRenderer/Private/SlateRHIRenderer.cpp b/Engine/Source/Runtime/SlateRHIRenderer/Private/SlateRHIRenderer.cpp
--- a/Engine/Source/Runtime/SlateRHIRenderer/Private/SlateRHIRenderer.cpp
+++ b/Engine/Source/Runtime/SlateRHIRenderer/Private/SlateRHIRenderer.cpp
@@ -1166,7 +1166,7 @@ void FSlateRHIRenderer::ReleaseUpdatableTexture(FSlateUpdatableTexture* Texture)
 {
 	if (IsInRenderingThread())
 	{
-		Texture->GetRenderResource()->InitResource();
+		Texture->GetRenderResource()->ReleaseResource();
 	}
 	else
 	{
