Merge Slate RHI stall fixes

diff --git a/Engine/Source/Runtime/SlateRHIRenderer/Private/SlateRHIRenderer.cpp b/Engine/Source/Runtime/SlateRHIRenderer/Private/SlateRHIRenderer.cpp
--- a/Engine/Source/Runtime/SlateRHIRenderer/Private/SlateRHIRenderer.cpp
+++ b/Engine/Source/Runtime/SlateRHIRenderer/Private/SlateRHIRenderer.cpp
@@ -400,8 +400,8 @@ void FSlateRHIRenderer::DrawWindow_RenderThread(FRHICommandListImmediate& RHICmd
 	{
 		SCOPE_CYCLE_COUNTER( STAT_SlateRenderingRTTime );
 
-		// Update the vertex and index buffer
-		RenderingPolicy->UpdateBuffers( WindowElementList );
+		// Update the vertex and index buffer		
+		RenderingPolicy->UpdateVertexAndIndexBuffers(RHICmdList, WindowElementList);
 		// should have been created by the game thread
 		check( IsValidRef(ViewportInfo.ViewportRHI) );
 
diff --git a/Engine/Source/Runtime/SlateRHIRenderer/Private/SlateRHIRenderingPolicy.cpp b/Engine/Source/Runtime/SlateRHIRenderer/Private/SlateRHIRenderingPolicy.cpp
--- a/Engine/Source/Runtime/SlateRHIRenderer/Private/SlateRHIRenderingPolicy.cpp
+++ b/Engine/Source/Runtime/SlateRHIRenderer/Private/SlateRHIRenderingPolicy.cpp
@@ -60,24 +60,36 @@ void FSlateElementIndexBuffer::ResizeBuffer( uint32 NewSizeBytes )
 	}
 }
 
-void FSlateElementIndexBuffer::FillBuffer( const TArray<SlateIndex>& InIndices, bool bShrinkToFit  )
+void FSlateElementIndexBuffer::PreFillBuffer(const TArray<SlateIndex>& InIndices, bool bShrinkToFit)
 {
-	check( IsInRenderingThread() );
+	check(IsInRenderingThread());
 
-	if( InIndices.Num() )
+	if (InIndices.Num())
 	{
 		uint32 NumIndices = InIndices.Num();
 
 		uint32 RequiredBufferSize = NumIndices*sizeof(SlateIndex);
 
 		// resize if needed
-		if( RequiredBufferSize > GetBufferSize() || bShrinkToFit )
+		if (RequiredBufferSize > GetBufferSize() || bShrinkToFit)
 		{
 			// Use array resize techniques for the vertex buffer
-			ResizeBuffer( InIndices.GetAllocatedSize() );
+			ResizeBuffer(InIndices.GetAllocatedSize());
 		}
 
-		BufferUsageSize += RequiredBufferSize;
+		BufferUsageSize += RequiredBufferSize;		
+	}
+}
+
+void FSlateElementIndexBuffer::FillBuffer_RenderThread( const TArray<SlateIndex>& InIndices )
+{
+	check( IsInRenderingThread() );
+
+	if( InIndices.Num() )
+	{
+		uint32 NumIndices = InIndices.Num();
+
+		uint32 RequiredBufferSize = NumIndices*sizeof(SlateIndex);		
 
 		void* IndicesPtr = RHILockIndexBuffer( IndexBufferRHI, 0, RequiredBufferSize, RLM_WriteOnly );
 
@@ -87,6 +99,23 @@ void FSlateElementIndexBuffer::FillBuffer( const TArray<SlateIndex>& InIndices,
 	}
 }
 
+void FSlateElementIndexBuffer::FillBuffer_RHIThread(const TArray<SlateIndex>& InIndices )
+{
+	check(IsInRHIThread() || (!IsInRenderingThread() && !IsInGameThread()));
+
+	if (InIndices.Num())
+	{
+		uint32 NumIndices = InIndices.Num();
+
+		uint32 RequiredBufferSize = NumIndices*sizeof(SlateIndex);		
+
+		//use non-flushing DynamicRHI version when on RHIThread.
+		void* IndicesPtr = GDynamicRHI->RHILockIndexBuffer(IndexBufferRHI, 0, RequiredBufferSize, RLM_WriteOnly);
+		FMemory::Memcpy(IndicesPtr, InIndices.GetData(), RequiredBufferSize);
+		GDynamicRHI->RHIUnlockIndexBuffer(IndexBufferRHI);		
+	}
+}
+
 /** Releases the index buffers RHI resource. */
 void FSlateElementIndexBuffer::ReleaseDynamicRHI()
 {
@@ -174,6 +203,52 @@ void FSlateRHIRenderingPolicy::EndDrawingWindows()
 	SET_MEMORY_STAT( STAT_SlateIndexBufferMemory, TotalIndexBufferMemory );
 }
 
+struct FSlateUpdateVertexAndIndexBuffers : public FRHICommand<FSlateUpdateVertexAndIndexBuffers>
+{
+	TSlateElementVertexBuffer<FSlateVertex>& VertexBuffer;
+	FSlateElementIndexBuffer& IndexBuffer;
+	const TArray<FSlateVertex>& Vertices;
+	const TArray<SlateIndex>& Indices;	
+
+	FSlateUpdateVertexAndIndexBuffers(TSlateElementVertexBuffer<FSlateVertex>& InVertexBuffer, const TArray<FSlateVertex>& InVertices, FSlateElementIndexBuffer& InIndexBuffer, const TArray<SlateIndex>& InIndices)
+		: VertexBuffer(InVertexBuffer)
+		, IndexBuffer(InIndexBuffer)
+		, Vertices(InVertices)
+		, Indices(InIndices)		
+	{}
+
+	void Execute(FRHICommandListBase& CmdList)
+	{
+		VertexBuffer.FillBuffer_RHIThread(Vertices);
+		IndexBuffer.FillBuffer_RHIThread(Indices);
+	}
+};
+
+void FSlateRHIRenderingPolicy::UpdateVertexAndIndexBuffers(FRHICommandListImmediate& RHICmdList, const FSlateWindowElementList& WindowElementList)
+{
+	// Should only be called by the rendering thread
+	check(IsInRenderingThread());
+
+	const TArray<FSlateVertex>& Vertices = WindowElementList.GetBatchedVertices();
+	const TArray<SlateIndex>& Indices = WindowElementList.GetBatchedIndices();
+
+	TSlateElementVertexBuffer<FSlateVertex>& VertexBuffer = VertexBuffers[CurrentBufferIndex];
+	FSlateElementIndexBuffer& IndexBuffer = IndexBuffers[CurrentBufferIndex];
+	VertexBuffer.PreFillBuffer(Vertices, bShouldShrinkResources);
+	IndexBuffer.PreFillBuffer(Indices, bShouldShrinkResources);
+
+	if (!GRHIThread || RHICmdList.Bypass())
+	{
+		VertexBuffer.FillBuffer_RenderThread(Vertices);
+		IndexBuffer.FillBuffer_RenderThread(Indices);		
+	}
+	else
+	{
+		new (RHICmdList.AllocCommand<FSlateUpdateVertexAndIndexBuffers>()) FSlateUpdateVertexAndIndexBuffers(VertexBuffer, Vertices, IndexBuffer, Indices);
+	}
+
+}
+
 void FSlateRHIRenderingPolicy::UpdateBuffers( const FSlateWindowElementList& WindowElementList )
 {
 	SCOPE_CYCLE_COUNTER( STAT_SlateUpdateBufferRTTime );
@@ -186,9 +261,12 @@ void FSlateRHIRenderingPolicy::UpdateBuffers( const FSlateWindowElementList& Win
 
 		TSlateElementVertexBuffer<FSlateVertex>& VertexBuffer = VertexBuffers[CurrentBufferIndex];
 		FSlateElementIndexBuffer& IndexBuffer = IndexBuffers[CurrentBufferIndex];
+		
+		VertexBuffer.PreFillBuffer(Vertices, bShouldShrinkResources);
+		IndexBuffer.PreFillBuffer(Indices, bShouldShrinkResources);
 
-		VertexBuffer.FillBuffer( Vertices, bShouldShrinkResources );
-		IndexBuffer.FillBuffer( Indices, bShouldShrinkResources );
+		VertexBuffer.FillBuffer_RenderThread(Vertices);
+		IndexBuffer.FillBuffer_RenderThread(Indices);
 	}
 }
 
diff --git a/Engine/Source/Runtime/SlateRHIRenderer/Private/SlateRHIRenderingPolicy.h b/Engine/Source/Runtime/SlateRHIRenderer/Private/SlateRHIRenderingPolicy.h
--- a/Engine/Source/Runtime/SlateRHIRenderer/Private/SlateRHIRenderingPolicy.h
+++ b/Engine/Source/Runtime/SlateRHIRenderer/Private/SlateRHIRenderingPolicy.h
@@ -64,30 +64,44 @@ public:
 	/** Resets the usage of the buffer */
 	void ResetBufferUsage() { BufferUsageSize = 0; }
 
-	/** Fills the buffer with slate vertices */
-	void FillBuffer( const TArray<VertexType>& InVertices, bool bShrinkToFit )
+	/** Resizes buffer, accumulates states safely on render thread */
+	void PreFillBuffer(const TArray<VertexType>& InVertices, bool bShrinkToFit)
 	{
-		check( IsInRenderingThread() );
+		check(IsInRenderingThread());
+		check(IsInRenderingThread());
 
-		if( InVertices.Num() )
+		if (InVertices.Num())
 		{
 			uint32 NumVertices = InVertices.Num();
 
-	#if !SLATE_USE_32BIT_INDICES
+#if !SLATE_USE_32BIT_INDICES
 			// make sure our index buffer can handle this
 			checkf(NumVertices < 0xFFFF, TEXT("Slate vertex buffer is too large (%d) to work with uint16 indices"), NumVertices);
-	#endif
+#endif
 
 			uint32 RequiredBufferSize = NumVertices*sizeof(VertexType);
 
 			// resize if needed
-			if( RequiredBufferSize > GetBufferSize() || bShrinkToFit )
+			if (RequiredBufferSize > GetBufferSize() || bShrinkToFit)
 			{
 				// Use array resize techniques for the vertex buffer
-				ResizeBuffer( InVertices.GetAllocatedSize() );
+				ResizeBuffer(InVertices.GetAllocatedSize());
 			}
 
 			BufferUsageSize += RequiredBufferSize;
+		}
+
+	}
+
+	/** Fills the buffer with slate vertices on renderthread. */
+	void FillBuffer_RenderThread( const TArray<VertexType>& InVertices)
+	{
+		check( IsInRenderingThread() );
+
+		if( InVertices.Num() )
+		{
+			uint32 NumVertices = InVertices.Num();
+			uint32 RequiredBufferSize = NumVertices*sizeof(VertexType);			
 
 			void* VerticesPtr = RHILockVertexBuffer( VertexBufferRHI, 0, RequiredBufferSize, RLM_WriteOnly );
 
@@ -97,6 +111,24 @@ public:
 		}
 	}
 
+	
+	void FillBuffer_RHIThread(const TArray<VertexType>& InVertices)
+	{
+		check(IsInRHIThread() || (!IsInRenderingThread() && !IsInGameThread()));
+
+		if (InVertices.Num())
+		{
+			uint32 NumVertices = InVertices.Num();
+
+			uint32 RequiredBufferSize = NumVertices*sizeof(VertexType);
+
+			//use non-flushing DynamicRHI version when on RHIThread.
+			void* VerticesPtr = GDynamicRHI->RHILockVertexBuffer(VertexBufferRHI, 0, RequiredBufferSize, RLM_WriteOnly);
+			FMemory::Memcpy(VerticesPtr, InVertices.GetData(), RequiredBufferSize);
+			GDynamicRHI->RHIUnlockVertexBuffer(VertexBufferRHI);			
+		}
+	}
+
 private:
 	/** Resizes the buffer to the passed in size.  Preserves internal data*/
 	void ResizeBuffer( uint32 NewSizeBytes )
@@ -152,8 +184,15 @@ public:
 	/** Resets the usage of the buffer */
 	void ResetBufferUsage() { BufferUsageSize = 0; }
 
+	/** Resizes buffer, accumulates states safely on render thread */
+	void PreFillBuffer(const TArray<SlateIndex>& InIndices, bool bShrinkToFit);
+
 	/** Fills the buffer with slate vertices */
-	void FillBuffer( const TArray<SlateIndex>& InIndices, bool bShrinkToFit );
+	void FillBuffer_RenderThread( const TArray<SlateIndex>& InIndices );
+
+	/** Fills the buffer with slate vertices */
+	void FillBuffer_RHIThread(const TArray<SlateIndex>& InIndices);
+
 private:
 	/** Resizes the buffer to the passed in size.  Preserves internal data */
 	void ResizeBuffer( uint32 NewSizeBytes );
@@ -177,6 +216,7 @@ public:
 	~FSlateRHIRenderingPolicy();
 
 	virtual void UpdateBuffers( const FSlateWindowElementList& WindowElementList ) override;
+	void UpdateVertexAndIndexBuffers(FRHICommandListImmediate& RHICmdList, const FSlateWindowElementList& WindowElementList);
 	virtual void DrawElements(FRHICommandListImmediate& RHICmdList, class FSlateBackBuffer& BackBuffer, const FMatrix& ViewProjectionMatrix, const TArray<FSlateRenderBatch>& RenderBatches, bool bAllowSwtichVerticalAxis=true);
 
 	virtual TSharedRef<FSlateFontCache> GetFontCache() override { return FontCache.ToSharedRef(); }
