Added another missing stat

diff --git a/Engine/Source/Runtime/SlateRHIRenderer/Private/SlateRHIRenderingPolicy.cpp b/Engine/Source/Runtime/SlateRHIRenderer/Private/SlateRHIRenderingPolicy.cpp
--- a/Engine/Source/Runtime/SlateRHIRenderer/Private/SlateRHIRenderingPolicy.cpp
+++ b/Engine/Source/Runtime/SlateRHIRenderer/Private/SlateRHIRenderingPolicy.cpp
@@ -12,7 +12,7 @@
 #include "SlateMaterialResource.h"
 
 DECLARE_CYCLE_STAT(TEXT("Update Buffers RT"), STAT_SlateUpdateBufferRTTime, STATGROUP_Slate);
-DECLARE_CYCLE_STAT(TEXT("Draw Time"), STAT_SlateDrawTime, STATGROUP_Slate);
+DECLARE_CYCLE_STAT(TEXT("Draw Time RT"), STAT_SlateDrawTime, STATGROUP_Slate);
 
 DECLARE_MEMORY_STAT(TEXT("Vertex Buffer Memory"), STAT_SlateVertexBufferMemory, STATGROUP_SlateMemory);
 DECLARE_MEMORY_STAT(TEXT("Index Buffer Memory"), STAT_SlateIndexBufferMemory, STATGROUP_SlateMemory);
@@ -226,6 +226,8 @@ struct FSlateUpdateVertexAndIndexBuffers : public FRHICommand<FSlateUpdateVertex
 
 void FSlateRHIRenderingPolicy::UpdateVertexAndIndexBuffers(FRHICommandListImmediate& RHICmdList, const FSlateWindowElementList& WindowElementList)
 {
+	SCOPE_CYCLE_COUNTER( STAT_SlateUpdateBufferRTTime );
+
 	// Should only be called by the rendering thread
 	check(IsInRenderingThread());
 
