Prevent a crash on shutdown in windows in open gl by setting the current gl display context to null when the current device context is being released back to the system.  This also required a small change to frame profiling to guard against the null context, and a change to the slate font texture which could end up processing a texture with a null context.

diff --git a/Engine/Source/Editor/MainFrame/Private/Frame/MainFrameHandler.cpp b/Engine/Source/Editor/MainFrame/Private/Frame/MainFrameHandler.cpp
--- a/Engine/Source/Editor/MainFrame/Private/Frame/MainFrameHandler.cpp
+++ b/Engine/Source/Editor/MainFrame/Private/Frame/MainFrameHandler.cpp
@@ -52,13 +52,7 @@ void FMainFrameHandler::ShutDownEditor()
 	if (RootWindow.IsValid())
 	{
 		RootWindow->SetRequestDestroyWindowOverride(FRequestDestroyWindowOverride());
-		
-		// This is a low risk temporary hack for 4.8 to avoid a crash when shutting down the open project window when running
-		// in opengl on windows.  A proper, higher risk fix is coming which will replace this.
-		if (FApp::HasGameName())
-		{
-			RootWindow->RequestDestroyWindow();
-		}
+		RootWindow->RequestDestroyWindow();
 	}
 
 	// Save out any config settings for the editor so they don't get lost
diff --git a/Engine/Source/Runtime/OpenGLDrv/Private/OpenGLQuery.cpp b/Engine/Source/Runtime/OpenGLDrv/Private/OpenGLQuery.cpp
--- a/Engine/Source/Runtime/OpenGLDrv/Private/OpenGLQuery.cpp
+++ b/Engine/Source/Runtime/OpenGLDrv/Private/OpenGLQuery.cpp
@@ -439,7 +439,7 @@ void FOpenGLBufferedGPUTiming::EndTiming()
 				TimerQuery->bInvalidResource = true;
 			}
 
-			if (TimerQuery->bInvalidResource)
+			if (TimerQuery->bInvalidResource && PlatformOpenGLContextValid())
 			{
 				PlatformGetNewRenderQuery(&TimerQuery->Resource, &TimerQuery->ResourceContext);
 				TimerQuery->bInvalidResource = false;
diff --git a/Engine/Source/Runtime/OpenGLDrv/Private/Windows/OpenGLWindows.cpp b/Engine/Source/Runtime/OpenGLDrv/Private/Windows/OpenGLWindows.cpp
--- a/Engine/Source/Runtime/OpenGLDrv/Private/Windows/OpenGLWindows.cpp
+++ b/Engine/Source/Runtime/OpenGLDrv/Private/Windows/OpenGLWindows.cpp
@@ -404,6 +404,10 @@ void PlatformReleaseOpenGLContext(FPlatformOpenGLDevice* Device, FPlatformOpenGL
 
 	check(Context->DeviceContext);
 
+	if (wglGetCurrentDC() == Context->DeviceContext)
+	{
+		wglMakeCurrent( NULL, NULL );
+	}
 	ReleaseDC(Context->WindowHandle, Context->DeviceContext);
 	Context->DeviceContext = NULL;
 
@@ -535,8 +539,11 @@ void PlatformSharedContextSetup(FPlatformOpenGLDevice* Device)
 
 void PlatformNULLContextSetup()
 {
-	// no need to glFlush() on Windows, it does flush by itself before switching contexts
-	ContextMakeCurrent(NULL, NULL);
+	if (wglGetCurrentDC())
+	{
+		// no need to glFlush() on Windows, it does flush by itself before switching contexts
+		ContextMakeCurrent(NULL, NULL);
+	}
 }
 
 /**
diff --git a/Engine/Source/Runtime/SlateRHIRenderer/Private/SlateRHIFontTexture.cpp b/Engine/Source/Runtime/SlateRHIRenderer/Private/SlateRHIFontTexture.cpp
--- a/Engine/Source/Runtime/SlateRHIRenderer/Private/SlateRHIFontTexture.cpp
+++ b/Engine/Source/Runtime/SlateRHIRenderer/Private/SlateRHIFontTexture.cpp
@@ -53,33 +53,12 @@ void FSlateFontTextureRHI::InitDynamicRHI()
 
 		INC_MEMORY_STAT_BY(STAT_SlateTextureGPUMemory, Width*Height*GPixelFormats[PF_A8].BlockBytes);
 	}
-
-	// Restore previous data if it exists
-	if( TempData.Num() > 0 )
-	{
-		// Set texture data back to the previous state
-		uint32 Stride;
-		uint8* TextureData = (uint8*)RHILockTexture2D( GetTypedResource(), 0, RLM_WriteOnly, Stride, false );
-		FMemory::Memcpy( TextureData, TempData.GetData(), GPixelFormats[PF_A8].BlockBytes*Width*Height );
-		RHIUnlockTexture2D( GetTypedResource(), 0, false );
-
-		TempData.Empty();
-	}
 }
 
 void FSlateFontTextureRHI::ReleaseDynamicRHI()
 {
 	check( IsInRenderingThread() );
 
-	// Copy the data to temporary storage until InitDynamicRHI is called
-	uint32 Stride;
-	TempData.Empty();
-	TempData.AddUninitialized( GPixelFormats[PF_A8].BlockBytes*Width*Height );
-
-	uint8* TextureData = (uint8*)RHILockTexture2D( GetTypedResource(), 0, RLM_ReadOnly, Stride, false );
-	FMemory::Memcpy( TempData.GetData(), TextureData, GPixelFormats[PF_A8].BlockBytes*Width*Height );
-	RHIUnlockTexture2D( GetTypedResource(), 0, false );
-
 	// Release the texture
 	if( IsValidRef(ShaderResource) )
 	{
