Fix compile errors on Mac and 16 bit index buffer platforms

diff --git a/Engine/Source/Runtime/SlateRHIRenderer/Private/SlateRHIRenderingPolicy.cpp b/Engine/Source/Runtime/SlateRHIRenderer/Private/SlateRHIRenderingPolicy.cpp
--- a/Engine/Source/Runtime/SlateRHIRenderer/Private/SlateRHIRenderingPolicy.cpp
+++ b/Engine/Source/Runtime/SlateRHIRenderer/Private/SlateRHIRenderingPolicy.cpp
@@ -136,7 +136,6 @@ FSlateRHIRenderingPolicy::FSlateRHIRenderingPolicy( TSharedPtr<FSlateFontCache>
 	, ResourceManager( InResourceManager )
 	, FontCache( InFontCache )
 	, CurrentBufferIndex(0)
-	, bShouldShrinkResources(false)
 	, bGammaCorrect(true)
 {
 	InitResources();
diff --git a/Engine/Source/Runtime/SlateRHIRenderer/Private/SlateRHIRenderingPolicy.h b/Engine/Source/Runtime/SlateRHIRenderer/Private/SlateRHIRenderingPolicy.h
--- a/Engine/Source/Runtime/SlateRHIRenderer/Private/SlateRHIRenderingPolicy.h
+++ b/Engine/Source/Runtime/SlateRHIRenderer/Private/SlateRHIRenderingPolicy.h
@@ -85,7 +85,7 @@ public:
 		{
 #if !SLATE_USE_32BIT_INDICES
 			// make sure our index buffer can handle this
-			checkf(RequiredVertexCount < 0xFFFF, TEXT("Slate vertex buffer is too large (%d) to work with uint16 indices"), NumVertices);
+			checkf(RequiredVertexCount < 0xFFFF, TEXT("Slate vertex buffer is too large (%d) to work with uint16 indices"), RequiredVertexCount);
 #endif
 			int32 RequiredBufferSize = RequiredVertexCount*sizeof(VertexType);
 
@@ -268,9 +268,6 @@ private:
 
 	uint8 CurrentBufferIndex;
 	
-	/** If we should shrink resources that are no longer used (do not set this from the game thread)*/
-	bool bShouldShrinkResources;
-
 	bool bGammaCorrect;
 };
 
