Slate - Attempting to find cases where we're releasing UTextures used in slate by allowing them to be garbage collected, but then attempting to use them later in slate with free'ed memory pointed to by the slate brush.  When this happens we ensure and use a fallback texture instead in non-shipping builds.

diff --git a/Engine/Source/Runtime/SlateRHIRenderer/Private/SlateRHIResourceManager.cpp b/Engine/Source/Runtime/SlateRHIRenderer/Private/SlateRHIResourceManager.cpp
--- a/Engine/Source/Runtime/SlateRHIRenderer/Private/SlateRHIResourceManager.cpp
+++ b/Engine/Source/Runtime/SlateRHIRenderer/Private/SlateRHIResourceManager.cpp
@@ -5,6 +5,7 @@
 #include "SlateNativeTextureResource.h"
 #include "SlateUTextureResource.h"
 #include "SlateMaterialResource.h"
+#include "ImageUtils.h"
 
 DECLARE_DWORD_ACCUMULATOR_STAT(TEXT("Num Texture Atlases"), STAT_SlateNumTextureAtlases, STATGROUP_SlateMemory);
 DECLARE_DWORD_ACCUMULATOR_STAT(TEXT("Num Non-Atlased Textures"), STAT_SlateNumNonAtlasedTextures, STATGROUP_SlateMemory);
@@ -109,6 +110,7 @@ void FDynamicResourceMap::ReleaseResources()
 
 
 FSlateRHIResourceManager::FSlateRHIResourceManager()
+	: BadResourceTexture(nullptr)
 {
 	FCoreDelegates::OnPreExit.AddRaw( this, &FSlateRHIResourceManager::OnAppExit );
 
@@ -515,25 +517,40 @@ FSlateShaderResourceProxy* FSlateRHIResourceManager::FindOrCreateDynamicTextureR
 	const FName ResourceName = InBrush.GetResourceName();
 	if ( ResourceName.IsValid() && ResourceName != NAME_None )
 	{
-		if( InBrush.GetResourceObject() != nullptr )
+		if ( UObject* ResourceObject = InBrush.GetResourceObject() )
 		{
-			UTexture2D* TextureObject = CastChecked<UTexture2D>(InBrush.GetResourceObject());
+			if ( !ResourceObject->IsA<UTexture2D>() )
+			{
+#if !(UE_BUILD_SHIPPING || UE_BUILD_TEST)
+				static TSet<UObject*> FailedTextures;
+				if ( !FailedTextures.Contains(ResourceObject) )
+				{
+					FailedTextures.Add(ResourceObject);
+					ensureMsgf(false, TEXT("Slate RHI Error - Invalid Texture2D '%s'."), *ResourceName.ToString());
+				}
+				ResourceObject = GetBadResourceTexture();
+#else
+				return nullptr;
+#endif
+			}
+
+			UTexture2D* TextureObject = Cast<UTexture2D>(ResourceObject);
 
 			TSharedPtr<FSlateUTextureResource> TextureResource = DynamicResourceMap.GetUTextureResource(TextureObject);
 
-			if(!TextureResource.IsValid())
+			if ( !TextureResource.IsValid() )
 			{
 				TextureResource = MakeDynamicUTextureResource(TextureObject);
-				if(TextureResource.IsValid())
+				if ( TextureResource.IsValid() )
 				{
 					INC_DWORD_STAT_BY(STAT_SlateNumDynamicTextures, 1);
 				}
 			}
 
-			if(TextureResource.IsValid())
+			if ( TextureResource.IsValid() )
 			{
 				UTexture2D* Texture2DObject = TextureResource->TextureObject;
-				if(Texture2DObject && !AccessedUTextures.Contains(Texture2DObject) && Texture2DObject->Resource)
+				if ( Texture2DObject && !AccessedUTextures.Contains(Texture2DObject) && Texture2DObject->Resource )
 				{
 					// Set the texture rendering resource that should be used.  The UTexture resource could change at any time so we must do this each frame
 					ENQUEUE_UNIQUE_RENDER_COMMAND_ONEPARAMETER(UpdateSlateUTextureResource,
@@ -545,11 +562,11 @@ FSlateShaderResourceProxy* FSlateRHIResourceManager::FindOrCreateDynamicTextureR
 							RenderTexture->LastRenderTime = FApp::GetCurrentTime();
 
 							// Refresh FTexture
-							InUTextureResource->UpdateRenderResource( RenderTexture );
+							InUTextureResource->UpdateRenderResource(RenderTexture);
 
 						});
 
-				
+
 					AccessedUTextures.Add(Texture2DObject);
 				}
 
@@ -582,7 +599,7 @@ FSlateShaderResourceProxy* FSlateRHIResourceManager::FindOrCreateDynamicTextureR
 	}
 
 	// dynamic texture was not found or loaded
-	return  nullptr;
+	return nullptr;
 }
 
 FSlateShaderResourceProxy* FSlateRHIResourceManager::GetMaterialResource(const FSlateBrush& InBrush)
@@ -786,3 +803,14 @@ void FSlateRHIResourceManager::ReloadTextures()
 	// Reload everythng
 	LoadUsedTextures();
 }
+
+UTexture* FSlateRHIResourceManager::GetBadResourceTexture()
+{
+	if ( BadResourceTexture == nullptr )
+	{
+		BadResourceTexture = FImageUtils::CreateCheckerboardTexture(FColor(255, 0, 255), FColor(255, 255, 0));
+		BadResourceTexture->AddToRoot();
+	}
+
+	return BadResourceTexture;
+}
diff --git a/Engine/Source/Runtime/SlateRHIRenderer/Private/SlateRHIResourceManager.h b/Engine/Source/Runtime/SlateRHIRenderer/Private/SlateRHIResourceManager.h
--- a/Engine/Source/Runtime/SlateRHIRenderer/Private/SlateRHIResourceManager.h
+++ b/Engine/Source/Runtime/SlateRHIRenderer/Private/SlateRHIResourceManager.h
@@ -177,6 +177,11 @@ private:
 	 */
 	void OnAppExit();
 
+	/**
+	 * Get or create the bad resource texture.
+	 */
+	UTexture* GetBadResourceTexture();
+
 private:
 	/** Map of all active dynamic resources being used by brushes */
 	FDynamicResourceMap DynamicResourceMap;
@@ -196,6 +201,7 @@ private:
 	uint32 AtlasSize;
 	/** This max size of each texture in an atlas */
 	FIntPoint MaxAltasedTextureSize;
-
+	/** Needed for displaying an error texture when we end up with bad resources. */
+	UTexture* BadResourceTexture;
 };
 
