Slate - Updating the 3D rendering path the properly set the RT and clear at the same time it to fix a bug on iOS.

diff --git a/Engine/Source/Runtime/SlateRHIRenderer/Private/Slate3DRenderer.cpp b/Engine/Source/Runtime/SlateRHIRenderer/Private/Slate3DRenderer.cpp
--- a/Engine/Source/Runtime/SlateRHIRenderer/Private/Slate3DRenderer.cpp
+++ b/Engine/Source/Runtime/SlateRHIRenderer/Private/Slate3DRenderer.cpp
@@ -92,8 +92,12 @@ void FSlate3DRenderer::DrawWindowToTarget_RenderThread( FRHICommandListImmediate
 	RenderTargetPolicy->UpdateVertexAndIndexBuffers( RHICmdList, BatchData );
 	FTextureRenderTarget2DResource* RenderTargetResource = static_cast<FTextureRenderTarget2DResource*>( RenderTarget->GetRenderTargetResource() );
 	
-	SetRenderTarget( RHICmdList, RenderTargetResource->GetTextureRHI(), FTextureRHIParamRef() );
-	RHICmdList.Clear( true, RenderTarget->ClearColor, false, 0.0f, false, 0x00, FIntRect() );
+	// Set render target and clear.
+	FRHIRenderTargetView ColorRTV(RenderTargetResource->GetTextureRHI());
+	FRHISetRenderTargetsInfo Info(1, &ColorRTV, FTextureRHIParamRef());
+	Info.ClearColors[0] = RenderTarget->ClearColor;
+
+	RHICmdList.SetRenderTargetsAndClear(Info);
 
 	FMatrix ProjectionMatrix = FSlateRHIRenderer::CreateProjectionMatrix( RenderTarget->SizeX, RenderTarget->SizeY );
 	if (BatchData.GetRenderBatches().Num() > 0)
