UMG - Additional fixes for 2561849, setting the clear color flag.

diff --git a/Engine/Source/Runtime/SlateRHIRenderer/Private/Slate3DRenderer.cpp b/Engine/Source/Runtime/SlateRHIRenderer/Private/Slate3DRenderer.cpp
--- a/Engine/Source/Runtime/SlateRHIRenderer/Private/Slate3DRenderer.cpp
+++ b/Engine/Source/Runtime/SlateRHIRenderer/Private/Slate3DRenderer.cpp
@@ -95,6 +95,7 @@ void FSlate3DRenderer::DrawWindowToTarget_RenderThread( FRHICommandListImmediate
 	// Set render target and clear.
 	FRHIRenderTargetView ColorRTV(RenderTargetResource->GetTextureRHI());
 	FRHISetRenderTargetsInfo Info(1, &ColorRTV, FTextureRHIParamRef());
+	Info.bClearColor = true;
 	Info.ClearColors[0] = RenderTarget->ClearColor;
 
 	RHICmdList.SetRenderTargetsAndClear(Info);
