Fixed crash  resizing a viewport from the movie player thread

diff --git a/Engine/Source/Runtime/SlateRHIRenderer/Private/SlateRHIRenderer.cpp b/Engine/Source/Runtime/SlateRHIRenderer/Private/SlateRHIRenderer.cpp
--- a/Engine/Source/Runtime/SlateRHIRenderer/Private/SlateRHIRenderer.cpp
+++ b/Engine/Source/Runtime/SlateRHIRenderer/Private/SlateRHIRenderer.cpp
@@ -288,8 +288,11 @@ void FSlateRHIRenderer::ConditionalResizeViewport( FViewportInfo* ViewInfo, uint
 		// The viewport size we have doesn't match the requested size of the viewport.
 		// Resize it now.
 
-		// cannot resize the viewport while potentially using it.
-		FlushRenderingCommands();
+		if( IsInGameThread() && !IsInSlateThread() )
+		{
+			// cannot resize the viewport while potentially using it.
+			FlushRenderingCommands();
+		}
 
 		// Windows are allowed to be zero sized ( sometimes they are animating to/from zero for example)
 		// but viewports cannot be zero sized.  Use 8x8 as a reasonably sized viewport in this case.
