Fix additional crashes resizing a viewport during a loading movie.

diff --git a/Engine/Source/Runtime/SlateRHIRenderer/Private/SlateRHIRenderer.cpp b/Engine/Source/Runtime/SlateRHIRenderer/Private/SlateRHIRenderer.cpp
--- a/Engine/Source/Runtime/SlateRHIRenderer/Private/SlateRHIRenderer.cpp
+++ b/Engine/Source/Runtime/SlateRHIRenderer/Private/SlateRHIRenderer.cpp
@@ -283,16 +283,14 @@ void FSlateRHIRenderer::ConditionalResizeViewport( FViewportInfo* ViewInfo, uint
 {
 	check( IsThreadSafeForSlateRendering() );
 
-	if( ViewInfo && ( ViewInfo->Height != Height || ViewInfo->Width != Width ||  ViewInfo->bFullscreen != bFullscreen || !IsValidRef(ViewInfo->ViewportRHI) ) )
+	if( IsInGameThread() && !IsInSlateThread() && ViewInfo && ( ViewInfo->Height != Height || ViewInfo->Width != Width ||  ViewInfo->bFullscreen != bFullscreen || !IsValidRef(ViewInfo->ViewportRHI) ) )
 	{
 		// The viewport size we have doesn't match the requested size of the viewport.
 		// Resize it now.
 
-		if( IsInGameThread() && !IsInSlateThread() )
-		{
-			// cannot resize the viewport while potentially using it.
-			FlushRenderingCommands();
-		}
+		// cannot resize the viewport while potentially using it.
+		FlushRenderingCommands();
+	
 
 		// Windows are allowed to be zero sized ( sometimes they are animating to/from zero for example)
 		// but viewports cannot be zero sized.  Use 8x8 as a reasonably sized viewport in this case.
