UMG/Slate - Fixing a fiew issues with memory leaking in the Slate RHI resource manager.  There was never a pass to remove stale entries in the dynamic resource map, so now we occasionally purge it of stale entries.  Also some general cleanup in the widget component to ensure all resources were being released.

diff --git a/Engine/Source/Runtime/SlateRHIRenderer/Private/SlateRHIRenderer.cpp b/Engine/Source/Runtime/SlateRHIRenderer/Private/SlateRHIRenderer.cpp
--- a/Engine/Source/Runtime/SlateRHIRenderer/Private/SlateRHIRenderer.cpp
+++ b/Engine/Source/Runtime/SlateRHIRenderer/Private/SlateRHIRenderer.cpp
@@ -571,6 +571,9 @@ void FSlateRHIRenderer::DrawWindows_Private( FSlateDrawBuffer& WindowDrawBuffer
 	// Update texture atlases if needed
 	ResourceManager->UpdateTextureAtlases();
 
+	// Remove expired assets from the dynamic resource pools.
+	ResourceManager->RemoveExpiredResources();
+
 	// Iterate through each element list and set up an RHI window for it if needed
 	TArray<TSharedPtr<FSlateWindowElementList>>& WindowElementLists = WindowDrawBuffer.GetWindowElementLists();
 	for( int32 ListIndex = 0; ListIndex < WindowElementLists.Num(); ++ListIndex )
diff --git a/Engine/Source/Runtime/SlateRHIRenderer/Private/SlateRHIResourceManager.cpp b/Engine/Source/Runtime/SlateRHIRenderer/Private/SlateRHIResourceManager.cpp
--- a/Engine/Source/Runtime/SlateRHIRenderer/Private/SlateRHIResourceManager.cpp
+++ b/Engine/Source/Runtime/SlateRHIRenderer/Private/SlateRHIResourceManager.cpp
@@ -11,6 +11,13 @@ DECLARE_DWORD_ACCUMULATOR_STAT(TEXT("Num Texture Atlases"), STAT_SlateNumTexture
 DECLARE_DWORD_ACCUMULATOR_STAT(TEXT("Num Non-Atlased Textures"), STAT_SlateNumNonAtlasedTextures, STATGROUP_SlateMemory);
 DECLARE_DWORD_ACCUMULATOR_STAT(TEXT("Num Dynamic Textures"), STAT_SlateNumDynamicTextures, STATGROUP_SlateMemory);
 
+FDynamicResourceMap::FDynamicResourceMap()
+	: LastExpiredTextureNumMarker(0)
+	, LastExpiredMaterialNumMarker(0)
+{
+
+}
+
 TSharedPtr<FSlateDynamicTextureResource> FDynamicResourceMap::GetDynamicTextureResource( FName ResourceName ) const
 {
 	return NativeTextureMap.FindRef( ResourceName );
@@ -20,7 +27,7 @@ TSharedPtr<FSlateUTextureResource> FDynamicResourceMap::GetUTextureResource( UTe
 {
 	if(TextureObject)
 	{
-		return UTextureResourceMap.FindRef(TextureObject);
+		return TextureMap.FindRef(TextureObject);
 	}
 
 	return nullptr;
@@ -28,7 +35,7 @@ TSharedPtr<FSlateUTextureResource> FDynamicResourceMap::GetUTextureResource( UTe
 
 TSharedPtr<FSlateMaterialResource> FDynamicResourceMap::GetMaterialResource( UMaterialInterface* Material ) const
 {
-	return MaterialResourceMap.FindRef( Material );
+	return MaterialMap.FindRef( Material );
 }
 
 
@@ -42,14 +49,14 @@ void FDynamicResourceMap::AddUTextureResource( UTexture2D* TextureObject, TShare
 	if(TextureObject)
 	{
 		check(TextureObject == InResource->TextureObject);
-		UTextureResourceMap.Add(TextureObject, InResource);
+		TextureMap.Add(TextureObject, InResource);
 	}
 }
 
 void FDynamicResourceMap::AddMaterialResource( UMaterialInterface* Material, TSharedRef<FSlateMaterialResource> InMaterialResource )
 {
 	check( Material == InMaterialResource->GetMaterialObject() );
-	MaterialResourceMap.Add( Material, InMaterialResource );
+	MaterialMap.Add(Material, InMaterialResource);
 }
 
 void FDynamicResourceMap::RemoveDynamicTextureResource(FName ResourceName)
@@ -61,14 +68,13 @@ void FDynamicResourceMap::RemoveUTextureResource( UTexture2D* TextureObject )
 {
 	if(TextureObject)
 	{
-		UTextureResourceMap.Remove(TextureObject);
+		TextureMap.Remove(TextureObject);
 	}
 }
 
-
 void FDynamicResourceMap::RemoveMaterialResource( UMaterialInterface* Material )
 {
-	MaterialResourceMap.Remove( Material );
+	MaterialMap.Remove(Material);
 }
 
 void FDynamicResourceMap::Empty()
@@ -85,15 +91,14 @@ void FDynamicResourceMap::EmptyDynamicTextureResources()
 
 void FDynamicResourceMap::EmptyUTextureResources()
 {
-	UTextureResourceMap.Empty();
+	TextureMap.Empty();
 }
 
 void FDynamicResourceMap::EmptyMaterialResources()
 {
-	MaterialResourceMap.Empty();
+	MaterialMap.Empty();
 }
 
-
 void FDynamicResourceMap::ReleaseResources()
 {
 	for (TMap<FName, TSharedPtr<FSlateDynamicTextureResource> >::TIterator It(NativeTextureMap); It; ++It)
@@ -101,13 +106,42 @@ void FDynamicResourceMap::ReleaseResources()
 		BeginReleaseResource(It.Value()->RHIRefTexture);
 	}
 	
-	for (TMap<TWeakObjectPtr<UTexture2D>, TSharedPtr<FSlateUTextureResource> >::TIterator It(UTextureResourceMap); It; ++It)
+	for ( TextureResourceMap::TIterator It(TextureMap); It; ++It )
 	{
 		It.Value()->UpdateRenderResource(nullptr);
 	}
-
 }
 
+void FDynamicResourceMap::RemoveExpiredResources()
+{
+	const int32 CheckingIncrement = 100;
+
+	if ( TextureMap.Num() > ( LastExpiredTextureNumMarker + CheckingIncrement ) )
+	{
+		for ( TextureResourceMap::TIterator It(TextureMap); It; ++It )
+		{
+			if ( It.Key().IsStale() )
+			{
+				It.RemoveCurrent();
+			}
+		}
+
+		LastExpiredTextureNumMarker = TextureMap.Num();
+	}
+
+	if ( MaterialMap.Num() > ( LastExpiredMaterialNumMarker + CheckingIncrement ) )
+	{
+		for ( MaterialResourceMap::TIterator It(MaterialMap); It; ++It )
+		{
+			if ( It.Key().IsStale() )
+			{
+				It.RemoveCurrent();
+			}
+		}
+
+		LastExpiredMaterialNumMarker = MaterialMap.Num();
+	}
+}
 
 FSlateRHIResourceManager::FSlateRHIResourceManager()
 	: BadResourceTexture(nullptr)
@@ -741,6 +775,11 @@ void FSlateRHIResourceManager::UpdateTextureAtlases()
 	}
 }
 
+void FSlateRHIResourceManager::RemoveExpiredResources()
+{
+	DynamicResourceMap.RemoveExpiredResources();
+}
+
 void FSlateRHIResourceManager::ReleaseResources()
 {
 	check( IsThreadSafeForSlateRendering() );
diff --git a/Engine/Source/Runtime/SlateRHIRenderer/Private/SlateRHIResourceManager.h b/Engine/Source/Runtime/SlateRHIRenderer/Private/SlateRHIResourceManager.h
--- a/Engine/Source/Runtime/SlateRHIRenderer/Private/SlateRHIResourceManager.h
+++ b/Engine/Source/Runtime/SlateRHIRenderer/Private/SlateRHIResourceManager.h
@@ -11,6 +11,8 @@ class FSlateMaterialResource;
 struct FDynamicResourceMap
 {
 public:
+	FDynamicResourceMap();
+
 	TSharedPtr<FSlateDynamicTextureResource> GetDynamicTextureResource( FName ResourceName ) const;
 
 	TSharedPtr<FSlateUTextureResource> GetUTextureResource( UTexture2D* TextureObject ) const;
@@ -34,14 +36,26 @@ public:
 
 	void ReleaseResources();
 
-	uint32 GetNumObjectResources() const { return UTextureResourceMap.Num() + MaterialResourceMap.Num(); }
+	uint32 GetNumObjectResources() const { return TextureMap.Num() + MaterialMap.Num(); }
+
+	void RemoveExpiredResources();
+
 private:
 	TMap<FName, TSharedPtr<FSlateDynamicTextureResource> > NativeTextureMap;
 	
-	TMap<TWeakObjectPtr<UTexture2D>, TSharedPtr<FSlateUTextureResource> > UTextureResourceMap;
+	typedef TMap<TWeakObjectPtr<UTexture2D>, TSharedPtr<FSlateUTextureResource> > TextureResourceMap;
+
+	/** Map of all texture resources */
+	TextureResourceMap TextureMap;
+
+	int32 LastExpiredTextureNumMarker;
+
+	typedef TMap<TWeakObjectPtr<UMaterialInterface>, TSharedPtr<FSlateMaterialResource> > MaterialResourceMap;
 
 	/** Map of all material resources */
-	TMap<TWeakObjectPtr<UMaterialInterface>, TSharedPtr<FSlateMaterialResource> > MaterialResourceMap;
+	MaterialResourceMap MaterialMap;
+
+	int32 LastExpiredMaterialNumMarker;
 };
 
 
@@ -80,6 +94,11 @@ public:
 	 */
 	void UpdateTextureAtlases();
 
+	/**
+	 * Removes entries from the DynamicResourceMap that have expired and are no longer valid UObject pointers.
+	 */
+	void RemoveExpiredResources();
+
 	/** FSlateShaderResourceManager interface */
 	virtual FSlateShaderResourceProxy* GetShaderResource( const FSlateBrush& InBrush ) override;
 	virtual ISlateAtlasProvider* GetTextureAtlasProvider() override;
diff --git a/Engine/Source/Runtime/UMG/Private/Components/WidgetComponent.cpp b/Engine/Source/Runtime/UMG/Private/Components/WidgetComponent.cpp
--- a/Engine/Source/Runtime/UMG/Private/Components/WidgetComponent.cpp
+++ b/Engine/Source/Runtime/UMG/Private/Components/WidgetComponent.cpp
@@ -50,10 +50,6 @@ public:
 		bWillEverBeLit = false;
 	}
 
-	~FWidget3DSceneProxy()
-	{
-	}
-
 	// FPrimitiveSceneProxy interface.
 	virtual void GetDynamicMeshElements(const TArray<const FSceneView*>& Views, const FSceneViewFamily& ViewFamily, uint32 VisibilityMap, FMeshElementCollector& Collector) const override
 	{
@@ -325,6 +321,7 @@ UWidgetComponent::UWidgetComponent( const FObjectInitializer& PCIP )
 	, TickWhenOffscreen( false )
 {
 	PrimaryComponentTick.bCanEverTick = true;
+	PrimaryComponentTick.TickGroup = TG_PostUpdateWork;
 	bTickInEditor = true;
 	
 	RelativeRotation = FRotator(0.f, 0.f, 90.f);
@@ -356,20 +353,16 @@ UWidgetComponent::UWidgetComponent( const FObjectInitializer& PCIP )
 	Space = EWidgetSpace::World;
 	Pivot = FVector2D(0.5, 0.5);
 	ZOrder = -100;
-
-	PrimaryComponentTick.TickGroup = TG_PostUpdateWork;
 }
 
 FPrimitiveSceneProxy* UWidgetComponent::CreateSceneProxy()
 {
-	if ( Space != EWidgetSpace::Screen )
+	if ( Space != EWidgetSpace::Screen && Renderer.IsValid() )
 	{
 		return new FWidget3DSceneProxy(this, *Renderer);
 	}
-	else
-	{
-		return nullptr;
-	}
+	
+	return nullptr;
 }
 
 FBoxSphereBounds UWidgetComponent::CalcBounds(const FTransform & LocalToWorld) const
@@ -499,14 +492,8 @@ void UWidgetComponent::OnUnregister()
 		}
 	}
 
-	if ( Widget )
-	{
-		Widget->RemoveFromParent();
-		Widget->MarkPendingKill();
-		Widget = nullptr;
-	}
+	ReleaseResources();
 
-	SlateWidget.Reset();
 	Super::OnUnregister();
 }
 
@@ -514,8 +501,11 @@ void UWidgetComponent::DestroyComponent(bool bPromoteChildren/*= false*/)
 {
 	Super::DestroyComponent(bPromoteChildren);
 
-	Renderer.Reset();
+	ReleaseResources();
+}
 
+void UWidgetComponent::ReleaseResources()
+{
 	if ( Widget )
 	{
 		Widget->RemoveFromParent();
@@ -523,7 +513,9 @@ void UWidgetComponent::DestroyComponent(bool bPromoteChildren/*= false*/)
 		Widget = nullptr;
 	}
 
+	Renderer.Reset();
 	SlateWidget.Reset();
+	HitTestGrid.Reset();
 }
 
 void UWidgetComponent::TickComponent(float DeltaTime, enum ELevelTick TickType, FActorComponentTickFunction *ThisTickFunction)
diff --git a/Engine/Source/Runtime/UMG/Public/Components/WidgetComponent.h b/Engine/Source/Runtime/UMG/Public/Components/WidgetComponent.h
--- a/Engine/Source/Runtime/UMG/Public/Components/WidgetComponent.h
+++ b/Engine/Source/Runtime/UMG/Public/Components/WidgetComponent.h
@@ -59,6 +59,9 @@ public:
 	/** Ensures the user widget is initialized */
 	void InitWidget();
 
+	/** Release resources associated with the widget. */
+	void ReleaseResources();
+
 	/** Ensures the 3d window is created its size and content. */
 	void UpdateWidget();
 
