UMG/Slate - Changing the removal of expired resources to be based on the amount of texture data being referened by the UI.

diff --git a/Engine/Source/Runtime/SlateRHIRenderer/Private/SlateRHIRenderer.cpp b/Engine/Source/Runtime/SlateRHIRenderer/Private/SlateRHIRenderer.cpp
--- a/Engine/Source/Runtime/SlateRHIRenderer/Private/SlateRHIRenderer.cpp
+++ b/Engine/Source/Runtime/SlateRHIRenderer/Private/SlateRHIRenderer.cpp
@@ -571,9 +571,6 @@ void FSlateRHIRenderer::DrawWindows_Private( FSlateDrawBuffer& WindowDrawBuffer
 	// Update texture atlases if needed
 	ResourceManager->UpdateTextureAtlases();
 
-	// Remove expired assets from the dynamic resource pools.
-	ResourceManager->RemoveExpiredResources();
-
 	// Iterate through each element list and set up an RHI window for it if needed
 	TArray<TSharedPtr<FSlateWindowElementList>>& WindowElementLists = WindowDrawBuffer.GetWindowElementLists();
 	for( int32 ListIndex = 0; ListIndex < WindowElementLists.Num(); ++ListIndex )
diff --git a/Engine/Source/Runtime/SlateRHIRenderer/Private/SlateRHIResourceManager.cpp b/Engine/Source/Runtime/SlateRHIRenderer/Private/SlateRHIResourceManager.cpp
--- a/Engine/Source/Runtime/SlateRHIRenderer/Private/SlateRHIResourceManager.cpp
+++ b/Engine/Source/Runtime/SlateRHIRenderer/Private/SlateRHIResourceManager.cpp
@@ -12,7 +12,7 @@ DECLARE_DWORD_ACCUMULATOR_STAT(TEXT("Num Non-Atlased Textures"), STAT_SlateNumNo
 DECLARE_DWORD_ACCUMULATOR_STAT(TEXT("Num Dynamic Textures"), STAT_SlateNumDynamicTextures, STATGROUP_SlateMemory);
 
 FDynamicResourceMap::FDynamicResourceMap()
-	: LastExpiredTextureNumMarker(0)
+	: TextureMemorySincePurge(0)
 	, LastExpiredMaterialNumMarker(0)
 {
 
@@ -46,10 +46,13 @@ void FDynamicResourceMap::AddDynamicTextureResource( FName ResourceName, TShared
 
 void FDynamicResourceMap::AddUTextureResource( UTexture2D* TextureObject, TSharedRef<FSlateUTextureResource> InResource)
 {
-	if(TextureObject)
+	if ( TextureObject )
 	{
 		check(TextureObject == InResource->TextureObject);
 		TextureMap.Add(TextureObject, InResource);
+
+		TextureMemorySincePurge += TextureObject->GetResourceSize(EResourceSizeMode::Inclusive);
+		RemoveExpiredTextureResources();
 	}
 }
 
@@ -57,6 +60,8 @@ void FDynamicResourceMap::AddMaterialResource( UMaterialInterface* Material, TSh
 {
 	check( Material == InMaterialResource->GetMaterialObject() );
 	MaterialMap.Add(Material, InMaterialResource);
+
+	RemoveExpiredMaterialResources();
 }
 
 void FDynamicResourceMap::RemoveDynamicTextureResource(FName ResourceName)
@@ -69,6 +74,7 @@ void FDynamicResourceMap::RemoveUTextureResource( UTexture2D* TextureObject )
 	if(TextureObject)
 	{
 		TextureMap.Remove(TextureObject);
+		TextureMemorySincePurge -= TextureObject->GetResourceSize(EResourceSizeMode::Inclusive);
 	}
 }
 
@@ -92,6 +98,7 @@ void FDynamicResourceMap::EmptyDynamicTextureResources()
 void FDynamicResourceMap::EmptyUTextureResources()
 {
 	TextureMap.Empty();
+	TextureMemorySincePurge = 0;
 }
 
 void FDynamicResourceMap::EmptyMaterialResources()
@@ -112,11 +119,12 @@ void FDynamicResourceMap::ReleaseResources()
 	}
 }
 
-void FDynamicResourceMap::RemoveExpiredResources()
+void FDynamicResourceMap::RemoveExpiredTextureResources()
 {
-	const int32 CheckingIncrement = 100;
+	// We attempt to purge every 10Mb of accumulated textures.
+	static const uint64 PurgeAfterAddingNewBytes = 1024 * 1024 * 10; // 10Mb
 
-	if ( TextureMap.Num() > ( LastExpiredTextureNumMarker + CheckingIncrement ) )
+	if ( TextureMemorySincePurge >= PurgeAfterAddingNewBytes )
 	{
 		for ( TextureResourceMap::TIterator It(TextureMap); It; ++It )
 		{
@@ -126,8 +134,13 @@ void FDynamicResourceMap::RemoveExpiredResources()
 			}
 		}
 
-		LastExpiredTextureNumMarker = TextureMap.Num();
+		TextureMemorySincePurge = 0;
 	}
+}
+
+void FDynamicResourceMap::RemoveExpiredMaterialResources()
+{
+	static const int32 CheckingIncrement = 10;
 
 	if ( MaterialMap.Num() > ( LastExpiredMaterialNumMarker + CheckingIncrement ) )
 	{
@@ -775,11 +788,6 @@ void FSlateRHIResourceManager::UpdateTextureAtlases()
 	}
 }
 
-void FSlateRHIResourceManager::RemoveExpiredResources()
-{
-	DynamicResourceMap.RemoveExpiredResources();
-}
-
 void FSlateRHIResourceManager::ReleaseResources()
 {
 	check( IsThreadSafeForSlateRendering() );
diff --git a/Engine/Source/Runtime/SlateRHIRenderer/Private/SlateRHIResourceManager.h b/Engine/Source/Runtime/SlateRHIRenderer/Private/SlateRHIResourceManager.h
--- a/Engine/Source/Runtime/SlateRHIRenderer/Private/SlateRHIResourceManager.h
+++ b/Engine/Source/Runtime/SlateRHIRenderer/Private/SlateRHIResourceManager.h
@@ -38,7 +38,9 @@ public:
 
 	uint32 GetNumObjectResources() const { return TextureMap.Num() + MaterialMap.Num(); }
 
-	void RemoveExpiredResources();
+private:
+	void RemoveExpiredTextureResources();
+	void RemoveExpiredMaterialResources();
 
 private:
 	TMap<FName, TSharedPtr<FSlateDynamicTextureResource> > NativeTextureMap;
@@ -48,7 +50,7 @@ private:
 	/** Map of all texture resources */
 	TextureResourceMap TextureMap;
 
-	int32 LastExpiredTextureNumMarker;
+	uint64 TextureMemorySincePurge;
 
 	typedef TMap<TWeakObjectPtr<UMaterialInterface>, TSharedPtr<FSlateMaterialResource> > MaterialResourceMap;
 
@@ -94,11 +96,6 @@ public:
 	 */
 	void UpdateTextureAtlases();
 
-	/**
-	 * Removes entries from the DynamicResourceMap that have expired and are no longer valid UObject pointers.
-	 */
-	void RemoveExpiredResources();
-
 	/** FSlateShaderResourceManager interface */
 	virtual FSlateShaderResourceProxy* GetShaderResource( const FSlateBrush& InBrush ) override;
 	virtual ISlateAtlasProvider* GetTextureAtlasProvider() override;
