Slate/UMG - When we update the RHI data, we now also update what the slate render proxy thinks is the actual size of the texture, which might change due to a reimport.

diff --git a/Engine/Source/Runtime/SlateRHIRenderer/Private/SlateUTextureResource.cpp b/Engine/Source/Runtime/SlateRHIRenderer/Private/SlateUTextureResource.cpp
--- a/Engine/Source/Runtime/SlateRHIRenderer/Private/SlateUTextureResource.cpp
+++ b/Engine/Source/Runtime/SlateRHIRenderer/Private/SlateUTextureResource.cpp
@@ -18,7 +18,7 @@ FSlateUTextureResource::FSlateUTextureResource(UTexture2D* InTexture)
 
 FSlateUTextureResource::~FSlateUTextureResource()
 {
-	if(Proxy)
+	if ( Proxy )
 	{
 		delete Proxy;
 	}
@@ -26,7 +26,19 @@ FSlateUTextureResource::~FSlateUTextureResource()
 
 void FSlateUTextureResource::UpdateRenderResource(FTexture* InFTexture)
 {
-	ShaderResource = InFTexture ? FTexture2DRHIRef(InFTexture->TextureRHI->GetTexture2D()) : FTexture2DRHIRef();
+	if ( InFTexture )
+	{
+		// If the RHI data has changed, it's possible the underlying size of the texture has changed,
+		// if that's true we need to update the actual size recorded on the proxy as well, otherwise 
+		// the texture will continue to render using the wrong size.
+		Proxy->ActualSize = FIntPoint(InFTexture->GetSizeX(), InFTexture->GetSizeY());
+
+		ShaderResource = FTexture2DRHIRef(InFTexture->TextureRHI->GetTexture2D());
+	}
+	else
+	{
+		ShaderResource = FTexture2DRHIRef();
+	}
 }
 
 uint32 FSlateUTextureResource::GetWidth() const
