Slate - Disabling the float16 precision version of the slate shader on all platforms.  Since the clip values are all in window space, it's terribly easy on modern monitors to exceed single digit precision (2048), causing multiple columns worth of additional pixel clipping erroneously.  Tweaking the way we pass that down to the vertex shader for compile so that it now properly algins with whatever we set the option to in C++.  We can consider turning it on for platforms where it wont matter, like the consoles 1920 will be your largest visible value.

diff --git a/Engine/Shaders/SlateVertexShader.usf b/Engine/Shaders/SlateVertexShader.usf
--- a/Engine/Shaders/SlateVertexShader.usf
+++ b/Engine/Shaders/SlateVertexShader.usf
@@ -16,19 +16,22 @@ struct VertexOut
 	float4 ClipExtents : TEXCOORD2;
 };
 
+
+
+
 VertexOut Main(
 	in float4 InTextureCoordinates : ATTRIBUTE0,
 /* Don't change these int streams to float, otherwise it WILL fail since the vertex declaration is defined as int on D3D11:
 		D3D11: WARNING: ID3D11Device::CreateInputLayout: The provided input signature expects to read an element with SemanticName/Index: 'ATTRIBUTE'/1 and component(s) of the type 'float32'.  However, the matching entry in the Input Layout declaration, element[1], specifies mismatched format: 'R16G16B16A16_SINT'.  This is not an error, since behavior is well defined: The element format determines what data conversion algorithm gets applied before it shows up in a shader register. Independently, the shader input signature defines how the shader will interpret the data that has been placed in its input registers, with no change in the bits stored.  It is valid for the application to reinterpret data as a different type once it is in the vertex shader, so this warning is issued just in case reinterpretation was not intended by the author. [ STATE_CREATION WARNING #391: CREATEINPUTLAYOUT_TYPE_MISMATCH ]
 */
-#if ES2_PROFILE && COMPILER_GLSL_ES2
-	in float2 InPosition : ATTRIBUTE1,
-	in float2 InClipOrigin : ATTRIBUTE2,
-	in float4 InClipExtents : ATTRIBUTE3,
-#else
+#if SLATE_USE_FLOAT16
 	in int2 InPosition : ATTRIBUTE1,
 	in half2 InClipOrigin : ATTRIBUTE2,
 	in half4 InClipExtents : ATTRIBUTE3,
+#else
+	in float2 InPosition : ATTRIBUTE1,
+	in float2 InClipOrigin : ATTRIBUTE2,
+	in float4 InClipExtents : ATTRIBUTE3,
 #endif
 	in float4 InColor : ATTRIBUTE4
 	)
@@ -41,7 +44,6 @@ VertexOut Main(
 	VOut.Position.y *= SwitchVerticalAxisMultiplier;
 #endif
 
-
 	// TextureCoordinates contains both the first and second texture coordinates in xy and zw respectively.
 	VOut.TextureCoordinates = InTextureCoordinates;
 	VOut.ClipOriginAndPos = float4(InClipOrigin, InPosition.xy);
diff --git a/Engine/Source/Runtime/SlateCore/Public/Rendering/RenderingCommon.h b/Engine/Source/Runtime/SlateCore/Public/Rendering/RenderingCommon.h
--- a/Engine/Source/Runtime/SlateCore/Public/Rendering/RenderingCommon.h
+++ b/Engine/Source/Runtime/SlateCore/Public/Rendering/RenderingCommon.h
@@ -8,7 +8,7 @@ enum class EPopupMethod : uint8;
 
 
 #define SLATE_USE_32BIT_INDICES !PLATFORM_USES_ES2
-#define SLATE_USE_FLOAT16 !PLATFORM_USES_ES2
+#define SLATE_USE_FLOAT16 0
 
 #if SLATE_USE_32BIT_INDICES
 typedef uint32 SlateIndex;
@@ -167,10 +167,18 @@ struct FSlateVertex
 {
 	/** Texture coordinates.  The first 2 are in xy and the 2nd are in zw */
 	float TexCoords[4]; 
+
+#if SLATE_USE_FLOAT16
 	/** Position of the vertex in window space */
 	int16 Position[2];
+#else
+	/** Position of the vertex in window space */
+	float Position[2];
+#endif
+
 	/** clip center/extents in render window space (window space with render transforms applied) */
 	FSlateRotatedClipRectType ClipRect;
+
 	/** Vertex color */
 	FColor Color;
 	
diff --git a/Engine/Source/Runtime/SlateRHIRenderer/Private/SlateShaders.cpp b/Engine/Source/Runtime/SlateRHIRenderer/Private/SlateShaders.cpp
--- a/Engine/Source/Runtime/SlateRHIRenderer/Private/SlateShaders.cpp
+++ b/Engine/Source/Runtime/SlateRHIRenderer/Private/SlateShaders.cpp
@@ -41,19 +41,20 @@ TGlobalResource<FSlateVertexDeclaration> GSlateVertexDeclaration;
 /************************************************************************/
 void FSlateVertexDeclaration::InitRHI()
 {
-	FVertexDeclarationElementList Elements;
-	uint32 Stride = sizeof(FSlateVertex);
-	Elements.Add(FVertexElement(0,STRUCT_OFFSET(FSlateVertex,TexCoords),VET_Float4,0,Stride));
-	Elements.Add(FVertexElement(0,STRUCT_OFFSET(FSlateVertex,Position),VET_Short2,1,Stride));
-	bool bUseFloat16 =
+	const bool bUseFloat16 =
 #if SLATE_USE_FLOAT16
 		true;
 #else
 		false;
 #endif
-	Elements.Add(FVertexElement(0,STRUCT_OFFSET(FSlateVertex,ClipRect),bUseFloat16 ? VET_Half2 : VET_Float2,2,Stride));
-	Elements.Add(FVertexElement(0,STRUCT_OFFSET(FSlateVertex,ClipRect)+STRUCT_OFFSET(FSlateRotatedClipRectType,ExtentX),bUseFloat16 ? VET_Half4 : VET_Float4,3,Stride));
-	Elements.Add(FVertexElement(0,STRUCT_OFFSET(FSlateVertex,Color),VET_Color,4,Stride));
+
+	FVertexDeclarationElementList Elements;
+	uint32 Stride = sizeof(FSlateVertex);
+	Elements.Add(FVertexElement(0, STRUCT_OFFSET(FSlateVertex, TexCoords), VET_Float4, 0, Stride));
+	Elements.Add(FVertexElement(0, STRUCT_OFFSET(FSlateVertex, Position), bUseFloat16 ? VET_Short2 : VET_Float2, 1, Stride));
+	Elements.Add(FVertexElement(0, STRUCT_OFFSET(FSlateVertex, ClipRect) + STRUCT_OFFSET(FSlateRotatedClipRectType, TopLeft), bUseFloat16 ? VET_Half2 : VET_Float2, 2, Stride));
+	Elements.Add(FVertexElement(0, STRUCT_OFFSET(FSlateVertex, ClipRect) + STRUCT_OFFSET(FSlateRotatedClipRectType, ExtentX), bUseFloat16 ? VET_Half4 : VET_Float4, 3, Stride));
+	Elements.Add(FVertexElement(0, STRUCT_OFFSET(FSlateVertex, Color), VET_Color, 4, Stride));
 
 	VertexDeclarationRHI = RHICreateVertexDeclaration(Elements);
 }
@@ -96,6 +97,17 @@ void FSlateElementVS::SetVerticalAxisMultiplier(FRHICommandList& RHICmdList, flo
 	SetShaderValue(RHICmdList, GetVertexShader(), SwitchVerticalAxisMultiplier, InMultiplier );
 }
 
+void FSlateElementVS::ModifyCompilationEnvironment(EShaderPlatform Platform, FShaderCompilerEnvironment& OutEnvironment)
+{
+	OutEnvironment.SetDefine(TEXT("SLATE_USE_FLOAT16"),
+#if SLATE_USE_FLOAT16
+		(uint32)1
+#else
+		(uint32)0
+#endif
+		);
+}
+
 /** Serializes the shader data */
 bool FSlateElementVS::Serialize( FArchive& Ar )
 {
diff --git a/Engine/Source/Runtime/SlateRHIRenderer/Private/SlateShaders.h b/Engine/Source/Runtime/SlateRHIRenderer/Private/SlateShaders.h
--- a/Engine/Source/Runtime/SlateRHIRenderer/Private/SlateShaders.h
+++ b/Engine/Source/Runtime/SlateRHIRenderer/Private/SlateShaders.h
@@ -41,6 +41,8 @@ public:
 
 	FSlateElementVS() {}
 
+	static void ModifyCompilationEnvironment(EShaderPlatform Platform, FShaderCompilerEnvironment& OutEnvironment);
+
 	/** 
 	 * Sets the view projection parameter
 	 *
