Slate - IT'S FIXED!  No more clip rounding error in Slate causing the last column or row of pixels to sometimes get clipped incorrectly.  In slate we do pixel snapping on each vertex individually, which means that a box that's at x.3 and extends to w.7, will visually rendered at x to (w+1).  The clipping rect was determined in pre-vertex rounding space, the clip rect is represented as origin, extentsX and extentsY.  The extents are all computed in pre-pixel snapping space, so a box that goes from x.3 to w.7, becomes an extents of z.4, and is then rounded down.  Meaning that on certain pixel offsets, some boxes ended up with an extents that is 1 pixel smaller from the visual version where each vertex is rounded individually.  The solution is to perform pixel snapping on the axis aligned clipping rect, which is then used to create the rotated clip rect which is formed with the extents, then to transform it, and again pixel snap.

Also removing the option to use halfs for clipping data when sending it to the pixel shader as it doesn't make a ton of sense to keep this codepath available after discussing it with the others working on Slate.

diff --git a/Engine/Shaders/SlateVertexShader.usf b/Engine/Shaders/SlateVertexShader.usf
--- a/Engine/Shaders/SlateVertexShader.usf
+++ b/Engine/Shaders/SlateVertexShader.usf
@@ -17,22 +17,14 @@ struct VertexOut
 };
 
 
-
-
 VertexOut Main(
 	in float4 InTextureCoordinates : ATTRIBUTE0,
 /* Don't change these int streams to float, otherwise it WILL fail since the vertex declaration is defined as int on D3D11:
 		D3D11: WARNING: ID3D11Device::CreateInputLayout: The provided input signature expects to read an element with SemanticName/Index: 'ATTRIBUTE'/1 and component(s) of the type 'float32'.  However, the matching entry in the Input Layout declaration, element[1], specifies mismatched format: 'R16G16B16A16_SINT'.  This is not an error, since behavior is well defined: The element format determines what data conversion algorithm gets applied before it shows up in a shader register. Independently, the shader input signature defines how the shader will interpret the data that has been placed in its input registers, with no change in the bits stored.  It is valid for the application to reinterpret data as a different type once it is in the vertex shader, so this warning is issued just in case reinterpretation was not intended by the author. [ STATE_CREATION WARNING #391: CREATEINPUTLAYOUT_TYPE_MISMATCH ]
 */
-#if SLATE_USE_FLOAT16
-	in int2 InPosition : ATTRIBUTE1,
-	in half2 InClipOrigin : ATTRIBUTE2,
-	in half4 InClipExtents : ATTRIBUTE3,
-#else
 	in float2 InPosition : ATTRIBUTE1,
 	in float2 InClipOrigin : ATTRIBUTE2,
 	in float4 InClipExtents : ATTRIBUTE3,
-#endif
 	in float4 InColor : ATTRIBUTE4
 	)
 {
diff --git a/Engine/Shaders/StandaloneRenderer/D3D/SlateDefaultVertexShader.hlsl b/Engine/Shaders/StandaloneRenderer/D3D/SlateDefaultVertexShader.hlsl
--- a/Engine/Shaders/StandaloneRenderer/D3D/SlateDefaultVertexShader.hlsl
+++ b/Engine/Shaders/StandaloneRenderer/D3D/SlateDefaultVertexShader.hlsl
@@ -17,10 +17,10 @@ struct VertexOut
 };
 
 VertexOut Main(
-	in int2 InPosition : POSITION,
+	in float2 InPosition : POSITION,
 	in float4 InTextureCoordinates : TEXCOORD0,
-	in half2 InClipOrigin : TEXCOORD1,
-	in half4 InClipExtents : TEXCOORD2,
+	in float2 InClipOrigin : TEXCOORD1,
+	in float4 InClipExtents : TEXCOORD2,
 	in float4 InColor : COLOR0
 	) 
 {
diff --git a/Engine/Source/Developer/StandaloneRenderer/Private/OpenGL/SlateOpenGLRenderingPolicy.cpp b/Engine/Source/Developer/StandaloneRenderer/Private/OpenGL/SlateOpenGLRenderingPolicy.cpp
--- a/Engine/Source/Developer/StandaloneRenderer/Private/OpenGL/SlateOpenGLRenderingPolicy.cpp
+++ b/Engine/Source/Developer/StandaloneRenderer/Private/OpenGL/SlateOpenGLRenderingPolicy.cpp
@@ -245,20 +245,13 @@ void FSlateOpenGLRenderingPolicy::DrawElements( const FMatrix& ViewProjectionMat
 		Offset = STRUCT_OFFSET( FSlateVertex, Position );
 		glVertexAttribPointer( 1, 2, GL_SHORT, GL_FALSE, Stride, BUFFER_OFFSET(Stride*BaseVertexIndex+Offset) );
 
-		bool bUseFloat16 =
-#if SLATE_USE_FLOAT16
-			true;
-#else
-			false;
-#endif
-
 		glEnableVertexAttribArray(2);
 		Offset = STRUCT_OFFSET( FSlateVertex, ClipRect );
-		glVertexAttribPointer( 2, 2, bUseFloat16 ? GL_HALF_FLOAT : GL_FLOAT, GL_FALSE, Stride, BUFFER_OFFSET(Stride*BaseVertexIndex+Offset) );
+		glVertexAttribPointer( 2, 2, GL_FLOAT, GL_FALSE, Stride, BUFFER_OFFSET(Stride*BaseVertexIndex+Offset) );
 
 		glEnableVertexAttribArray(3);
 		Offset = STRUCT_OFFSET( FSlateVertex, ClipRect ) + STRUCT_OFFSET( FSlateRotatedClipRectType, ExtentX );
-		glVertexAttribPointer( 3, 4, bUseFloat16 ? GL_HALF_FLOAT : GL_FLOAT, GL_FALSE, Stride, BUFFER_OFFSET(Stride*BaseVertexIndex+Offset) );
+		glVertexAttribPointer( 3, 4, GL_FLOAT, GL_FALSE, Stride, BUFFER_OFFSET(Stride*BaseVertexIndex+Offset) );
 
 		glEnableVertexAttribArray(4);
 		Offset = STRUCT_OFFSET( FSlateVertex, Color );
diff --git a/Engine/Source/Developer/StandaloneRenderer/Private/Windows/D3D/SlateD3DShaders.cpp b/Engine/Source/Developer/StandaloneRenderer/Private/Windows/D3D/SlateD3DShaders.cpp
--- a/Engine/Source/Developer/StandaloneRenderer/Private/Windows/D3D/SlateD3DShaders.cpp
+++ b/Engine/Source/Developer/StandaloneRenderer/Private/Windows/D3D/SlateD3DShaders.cpp
@@ -241,9 +241,9 @@ FSlateDefaultVS::FSlateDefaultVS()
 	D3D11_INPUT_ELEMENT_DESC Layout[] = 
 	{
 		{ "TEXCOORD",	0, DXGI_FORMAT_R32G32B32A32_FLOAT,	0, 0,								D3D11_INPUT_PER_VERTEX_DATA,	0 },
-		{ "POSITION",	0, DXGI_FORMAT_R16G16_SINT,			0, D3D11_APPEND_ALIGNED_ELEMENT,	D3D11_INPUT_PER_VERTEX_DATA,	0 },
-		{ "TEXCOORD",	1, DXGI_FORMAT_R16G16_FLOAT,		0, D3D11_APPEND_ALIGNED_ELEMENT,	D3D11_INPUT_PER_VERTEX_DATA,	0 },
-		{ "TEXCOORD",	2, DXGI_FORMAT_R16G16B16A16_FLOAT,	0, D3D11_APPEND_ALIGNED_ELEMENT,	D3D11_INPUT_PER_VERTEX_DATA,	0 },
+		{ "POSITION",	0, DXGI_FORMAT_R32G32_FLOAT,		0, D3D11_APPEND_ALIGNED_ELEMENT,	D3D11_INPUT_PER_VERTEX_DATA,	0 },
+		{ "TEXCOORD",	1, DXGI_FORMAT_R32G32_FLOAT,		0, D3D11_APPEND_ALIGNED_ELEMENT,	D3D11_INPUT_PER_VERTEX_DATA,	0 },
+		{ "TEXCOORD",	2, DXGI_FORMAT_R32G32B32A32_FLOAT,	0, D3D11_APPEND_ALIGNED_ELEMENT,	D3D11_INPUT_PER_VERTEX_DATA,	0 },
 		{ "COLOR",		0, DXGI_FORMAT_B8G8R8A8_UNORM,		0, D3D11_APPEND_ALIGNED_ELEMENT,	D3D11_INPUT_PER_VERTEX_DATA,	0 },
 	};
 
diff --git a/Engine/Source/Runtime/SlateCore/Private/Rendering/RenderingCommon.cpp b/Engine/Source/Runtime/SlateCore/Private/Rendering/RenderingCommon.cpp
--- a/Engine/Source/Runtime/SlateCore/Private/Rendering/RenderingCommon.cpp
+++ b/Engine/Source/Runtime/SlateCore/Private/Rendering/RenderingCommon.cpp
@@ -6,9 +6,9 @@
 FSlateRotatedRect::FSlateRotatedRect() {}
 
 FSlateRotatedRect::FSlateRotatedRect(const FSlateRect& AlignedRect)
-	: TopLeft(AlignedRect.GetTopLeft())
-	, ExtentX(AlignedRect.Right - AlignedRect.Left, 0.0f)
-	, ExtentY(0.0f, AlignedRect.Bottom - AlignedRect.Top)
+	: TopLeft(FMath::RoundToInt(AlignedRect.Left), FMath::RoundToInt(AlignedRect.Top))
+	, ExtentX(FMath::RoundToInt(AlignedRect.Right) - FMath::RoundToInt(AlignedRect.Left), 0.0f)
+	, ExtentY(0.0f, FMath::RoundToInt(AlignedRect.Bottom) - FMath::RoundToInt(AlignedRect.Top))
 {
 }
 
@@ -52,22 +52,6 @@ bool FSlateRotatedRect::IsUnderLocation(const FVector2D& Location) const
 	return false;
 }
 
-FSlateRotatedRectHalf::FSlateRotatedRectHalf() {}
-
-FSlateRotatedRectHalf::FSlateRotatedRectHalf(const FSlateRotatedRect& RotatedRect)
-	: TopLeft(RotatedRect.TopLeft)
-	, ExtentX(RotatedRect.ExtentX)
-	, ExtentY(RotatedRect.ExtentY)
-{
-}
-
-FSlateRotatedRectHalf::FSlateRotatedRectHalf(const FVector2D& InTopLeft, const FVector2D& InExtentX, const FVector2D& InExtentY)
-	: TopLeft(InTopLeft)
-	, ExtentX(InExtentX)
-	, ExtentY(InExtentY)
-{
-}
-
 FSlateVertex::FSlateVertex() 
 {
 }
diff --git a/Engine/Source/Runtime/SlateCore/Public/Rendering/RenderingCommon.h b/Engine/Source/Runtime/SlateCore/Public/Rendering/RenderingCommon.h
--- a/Engine/Source/Runtime/SlateCore/Public/Rendering/RenderingCommon.h
+++ b/Engine/Source/Runtime/SlateCore/Public/Rendering/RenderingCommon.h
@@ -8,7 +8,6 @@ enum class EPopupMethod : uint8;
 
 
 #define SLATE_USE_32BIT_INDICES !PLATFORM_USES_ES2
-#define SLATE_USE_FLOAT16 0
 
 #if SLATE_USE_32BIT_INDICES
 typedef uint32 SlateIndex;
@@ -130,35 +129,7 @@ FSlateRotatedRect TransformRect(const TransformType& Transform, const FSlateRota
 	);
 }
 
-
-/**
- * Stores a Rotated rect as float16 (for rendering).
- */
-struct FSlateRotatedRectHalf
-{
-	/** Default ctor. */
-	FSlateRotatedRectHalf();
-	/** Construct a float16 version of a rotated rect from a full-float version. */
-	explicit FSlateRotatedRectHalf(const FSlateRotatedRect& RotatedRect);
-	/** Per-element constructor. */
-	FSlateRotatedRectHalf(const FVector2D& InTopLeft, const FVector2D& InExtentX, const FVector2D& InExtentY);
-	/** transformed Top-left corner. */
-	FVector2DHalf TopLeft;
-	/** transformed X extent (right-left). */
-	FVector2DHalf ExtentX;
-	/** transformed Y extent (bottom-top). */
-	FVector2DHalf ExtentY;
-};
-
-/**
- * Not all platforms support Float16, so we have to be tricky here and declare the proper vertex type.
- */
-#if SLATE_USE_FLOAT16
-typedef FSlateRotatedRectHalf FSlateRotatedClipRectType;
-#else
 typedef FSlateRotatedRect FSlateRotatedClipRectType;
-#endif
-
 
 /** 
  * A struct which defines a basic vertex seen by the Slate vertex buffers and shaders
@@ -168,13 +139,8 @@ struct FSlateVertex
 	/** Texture coordinates.  The first 2 are in xy and the 2nd are in zw */
 	float TexCoords[4]; 
 
-#if SLATE_USE_FLOAT16
-	/** Position of the vertex in window space */
-	int16 Position[2];
-#else
 	/** Position of the vertex in window space */
 	float Position[2];
-#endif
 
 	/** clip center/extents in render window space (window space with render transforms applied) */
 	FSlateRotatedClipRectType ClipRect;
diff --git a/Engine/Source/Runtime/SlateRHIRenderer/Private/SlateShaders.cpp b/Engine/Source/Runtime/SlateRHIRenderer/Private/SlateShaders.cpp
--- a/Engine/Source/Runtime/SlateRHIRenderer/Private/SlateShaders.cpp
+++ b/Engine/Source/Runtime/SlateRHIRenderer/Private/SlateShaders.cpp
@@ -41,19 +41,12 @@ TGlobalResource<FSlateVertexDeclaration> GSlateVertexDeclaration;
 /************************************************************************/
 void FSlateVertexDeclaration::InitRHI()
 {
-	const bool bUseFloat16 =
-#if SLATE_USE_FLOAT16
-		true;
-#else
-		false;
-#endif
-
 	FVertexDeclarationElementList Elements;
 	uint32 Stride = sizeof(FSlateVertex);
 	Elements.Add(FVertexElement(0, STRUCT_OFFSET(FSlateVertex, TexCoords), VET_Float4, 0, Stride));
-	Elements.Add(FVertexElement(0, STRUCT_OFFSET(FSlateVertex, Position), bUseFloat16 ? VET_Short2 : VET_Float2, 1, Stride));
-	Elements.Add(FVertexElement(0, STRUCT_OFFSET(FSlateVertex, ClipRect) + STRUCT_OFFSET(FSlateRotatedClipRectType, TopLeft), bUseFloat16 ? VET_Half2 : VET_Float2, 2, Stride));
-	Elements.Add(FVertexElement(0, STRUCT_OFFSET(FSlateVertex, ClipRect) + STRUCT_OFFSET(FSlateRotatedClipRectType, ExtentX), bUseFloat16 ? VET_Half4 : VET_Float4, 3, Stride));
+	Elements.Add(FVertexElement(0, STRUCT_OFFSET(FSlateVertex, Position), VET_Float2, 1, Stride));
+	Elements.Add(FVertexElement(0, STRUCT_OFFSET(FSlateVertex, ClipRect) + STRUCT_OFFSET(FSlateRotatedClipRectType, TopLeft), VET_Float2, 2, Stride));
+	Elements.Add(FVertexElement(0, STRUCT_OFFSET(FSlateVertex, ClipRect) + STRUCT_OFFSET(FSlateRotatedClipRectType, ExtentX), VET_Float4, 3, Stride));
 	Elements.Add(FVertexElement(0, STRUCT_OFFSET(FSlateVertex, Color), VET_Color, 4, Stride));
 
 	VertexDeclarationRHI = RHICreateVertexDeclaration(Elements);
@@ -99,13 +92,6 @@ void FSlateElementVS::SetVerticalAxisMultiplier(FRHICommandList& RHICmdList, flo
 
 void FSlateElementVS::ModifyCompilationEnvironment(EShaderPlatform Platform, FShaderCompilerEnvironment& OutEnvironment)
 {
-	OutEnvironment.SetDefine(TEXT("SLATE_USE_FLOAT16"),
-#if SLATE_USE_FLOAT16
-		(uint32)1
-#else
-		(uint32)0
-#endif
-		);
 }
 
 /** Serializes the shader data */
