Fix crash when using a material on a font and the font cache had to be flushed

diff --git a/Engine/Source/Runtime/SlateRHIRenderer/Private/SlateRHIResourceManager.cpp b/Engine/Source/Runtime/SlateRHIRenderer/Private/SlateRHIResourceManager.cpp
--- a/Engine/Source/Runtime/SlateRHIRenderer/Private/SlateRHIResourceManager.cpp
+++ b/Engine/Source/Runtime/SlateRHIRenderer/Private/SlateRHIResourceManager.cpp
@@ -669,7 +669,7 @@ FSlateMaterialResource* FSlateRHIResourceManager::GetMaterialResource(const UObj
 	const UMaterialInterface* Material = CastChecked<UMaterialInterface>(InMaterial);
 
 	TSharedPtr<FSlateMaterialResource> MaterialResource = DynamicResourceMap.GetMaterialResource(Material);
-	if (!MaterialResource.IsValid() || MaterialResource->GetTextureMaskResource() != TextureMask )
+	if (!MaterialResource.IsValid())
 	{
 		// Get a resource from the free list if possible
 		if(MaterialResourceFreeList.Num() > 0)
@@ -685,6 +685,10 @@ FSlateMaterialResource* FSlateRHIResourceManager::GetMaterialResource(const UObj
 		
 		DynamicResourceMap.AddMaterialResource(Material, MaterialResource.ToSharedRef());
 	}
+	else if( MaterialResource->GetTextureMaskResource() != TextureMask )
+	{
+		MaterialResource->UpdateMaterial( *Material, ImageSize, TextureMask );
+	}
 	else
 	{
 		// Keep the resource up to date
