Custom UV region for FSlateBrush. When valid - slate will use brush UV instead of UV specified in resource proxy.
Slate resource manager accepts UTexture2Dynamic as a resource.

diff --git a/Engine/Source/Runtime/CoreUObject/Classes/Object.h b/Engine/Source/Runtime/CoreUObject/Classes/Object.h
--- a/Engine/Source/Runtime/CoreUObject/Classes/Object.h
+++ b/Engine/Source/Runtime/CoreUObject/Classes/Object.h
@@ -531,6 +531,24 @@ struct FBox
 
 };
 
+/**
+ * A rectangular 2D Box.
+ * The full C++ class is located here: Engine\Source\Runtime\Core\Public\Math\Box2D.h
+ */
+USTRUCT(immutable, noexport, BlueprintType)
+struct FBox2D
+{
+	UPROPERTY(EditAnywhere, BlueprintReadWrite, Category=Box2D, SaveGame)
+	FVector2D Min;
+
+	UPROPERTY(EditAnywhere, BlueprintReadWrite, Category=Box2D, SaveGame)
+	FVector2D Max;
+
+	UPROPERTY()
+	uint8 IsValid;
+
+};
+
 
 /**
  * A bounding box and bounding sphere with the same origin.
diff --git a/Engine/Source/Runtime/SlateCore/Private/Rendering/ElementBatcher.cpp b/Engine/Source/Runtime/SlateCore/Private/Rendering/ElementBatcher.cpp
--- a/Engine/Source/Runtime/SlateCore/Private/Rendering/ElementBatcher.cpp
+++ b/Engine/Source/Runtime/SlateCore/Private/Rendering/ElementBatcher.cpp
@@ -286,9 +286,20 @@ void FSlateElementBatcher::AddBoxElement( const FSlateDrawElement& DrawElement )
 			// Texel offset
 			HalfTexel = FVector2D( PixelCenterOffset/TextureWidth, PixelCenterOffset/TextureHeight );
 
-			SizeUV = ResourceProxy->SizeUV;
-			StartUV = ResourceProxy->StartUV + HalfTexel;
-			EndUV = StartUV + ResourceProxy->SizeUV;
+			FBox2D BrushUV = BrushResource->GetUVRegion();
+			//In case brush has valid UV region - use it instead of proxy UV
+			if (BrushUV.bIsValid)
+			{
+				SizeUV = BrushUV.GetSize();
+				StartUV = BrushUV.Min + HalfTexel;
+				EndUV = StartUV + SizeUV;
+			}
+			else
+			{
+				SizeUV = ResourceProxy->SizeUV;
+				StartUV = ResourceProxy->StartUV + HalfTexel;
+				EndUV = StartUV + ResourceProxy->SizeUV;
+			}
 		}
 		else
 		{
diff --git a/Engine/Source/Runtime/SlateCore/Private/Styling/SlateBrush.cpp b/Engine/Source/Runtime/SlateCore/Private/Styling/SlateBrush.cpp
--- a/Engine/Source/Runtime/SlateCore/Private/Styling/SlateBrush.cpp
+++ b/Engine/Source/Runtime/SlateCore/Private/Styling/SlateBrush.cpp
@@ -17,6 +17,7 @@ FSlateBrush::FSlateBrush( ESlateBrushDrawType::Type InDrawType, const FName InRe
 	, ResourceObject( InObjectResource )
 	, ResourceName( InResourceName )
 	, bIsDynamicallyLoaded( bInDynamicallyLoaded )
+	, UVRegion( ForceInit )
 {
 	bHasUObject_DEPRECATED = (InObjectResource != nullptr) || InResourceName.ToString().StartsWith(FSlateBrush::UTextureIdentifier());
 
@@ -39,6 +40,7 @@ FSlateBrush::FSlateBrush( ESlateBrushDrawType::Type InDrawType, const FName InRe
 	, ResourceObject( InObjectResource )
 	, ResourceName( InResourceName )
 	, bIsDynamicallyLoaded( bInDynamicallyLoaded )
+	, UVRegion( ForceInit )
 {
 	bHasUObject_DEPRECATED = (InObjectResource != nullptr) || InResourceName.ToString().StartsWith(FSlateBrush::UTextureIdentifier());
 
@@ -61,6 +63,7 @@ FSlateBrush::FSlateBrush( ESlateBrushDrawType::Type InDrawType, const FName InRe
 	, ResourceObject(InObjectResource)
 	, ResourceName(InResourceName)
 	, bIsDynamicallyLoaded(bInDynamicallyLoaded)
+	, UVRegion(ForceInit)
 {
 	bHasUObject_DEPRECATED = (InObjectResource != nullptr) || InResourceName.ToString().StartsWith(FSlateBrush::UTextureIdentifier());
 
diff --git a/Engine/Source/Runtime/SlateCore/Public/Brushes/SlateDynamicImageBrush.h b/Engine/Source/Runtime/SlateCore/Public/Brushes/SlateDynamicImageBrush.h
--- a/Engine/Source/Runtime/SlateCore/Public/Brushes/SlateDynamicImageBrush.h
+++ b/Engine/Source/Runtime/SlateCore/Public/Brushes/SlateDynamicImageBrush.h
@@ -28,23 +28,31 @@ struct SLATECORE_API FSlateDynamicImageBrush
 		, bRemoveResourceFromRootSet(false)
 	{
 		bIsDynamicallyLoaded = true;
-
-		// if we have a texture, make a unique name
-		if (ResourceObject != nullptr)
-		{
-			// @todo Slate - Hack:  This is to address an issue where the brush created and a GC occurs before the brush resource object becomes referenced
-			// by the Slate resource manager. Don't add objects that are already in root set (and mark them as such) to avoid incorrect removing objects
-			// from root set in destructor.
-			if (!ResourceObject->HasAllFlags(RF_RootSet))
-			{
-				ResourceObject->AddToRoot();
-				bRemoveResourceFromRootSet = true;
-			}
-
-			ResourceName = InTextureName;
-		}
+		InitFromTextureObject(InTextureName);
 	}
 
+	/**
+	 * @param InTexture		The UTexture2DDynamic being used for this brush.
+	 * @param InImageSize	How large should the image be (not necessarily the image size on disk)
+	 * @param InTint		The tint of the image
+	 * @param InTiling		How do we tile if at all?
+	 * @param InImageType	The type of image this this is
+	 */
+	FORCENOINLINE FSlateDynamicImageBrush( 
+		class UTexture2DDynamic* InTexture, 
+		const FVector2D& InImageSize,
+		const FName InTextureName,
+		const FLinearColor& InTint = FLinearColor(1.0f, 1.0f, 1.0f, 1.0f), 
+		ESlateBrushTileType::Type InTiling = ESlateBrushTileType::NoTile, 
+		ESlateBrushImageType::Type InImageType = ESlateBrushImageType::FullColor
+	)
+		: FSlateBrush(ESlateBrushDrawType::Image, FName(TEXT("None")), FMargin(0.0f), InTiling, InImageType, InImageSize, InTint, (UObject*)InTexture)
+		, bRemoveResourceFromRootSet(false)
+	{
+		bIsDynamicallyLoaded = true;
+		InitFromTextureObject(InTextureName);
+	}
+		
 	/**
 	 * @param InTextureName		The name of the texture to load.
 	 * @param InImageSize		How large should the image be (not necessarily the image size on disk)
@@ -84,6 +92,25 @@ struct SLATECORE_API FSlateDynamicImageBrush
 	/** Destructor. */
 	virtual ~FSlateDynamicImageBrush();
 
+private:
+	void InitFromTextureObject(FName InTextureName)
+	{
+		// if we have a texture, make a unique name
+		if (ResourceObject != nullptr)
+		{
+			// @todo Slate - Hack:  This is to address an issue where the brush created and a GC occurs before the brush resource object becomes referenced
+			// by the Slate resource manager. Don't add objects that are already in root set (and mark them as such) to avoid incorrect removing objects
+			// from root set in destructor.
+			if (!ResourceObject->HasAllFlags(RF_RootSet))
+			{
+				ResourceObject->AddToRoot();
+				bRemoveResourceFromRootSet = true;
+			}
+
+			ResourceName = InTextureName;
+		}
+	}
+
 private:
 	// Tracks if Resource was in root set to avoid unnecessary removing it from there.
 	bool bRemoveResourceFromRootSet;
diff --git a/Engine/Source/Runtime/SlateCore/Public/Styling/SlateBrush.h b/Engine/Source/Runtime/SlateCore/Public/Styling/SlateBrush.h
--- a/Engine/Source/Runtime/SlateCore/Public/Styling/SlateBrush.h
+++ b/Engine/Source/Runtime/SlateCore/Public/Styling/SlateBrush.h
@@ -156,6 +156,7 @@ public:
 		, ResourceName(NAME_None)
 		, bIsDynamicallyLoaded(false)
 		, bHasUObject_DEPRECATED(false)
+		, UVRegion(ForceInit)
 	{ }
 
 public:
@@ -223,6 +224,26 @@ public:
 		return bIsDynamicallyLoaded;
 	}
 
+	/**
+	 * Get brush UV region, should check if region is valid before using it
+	 *
+	 * @return UV region
+	 */
+	FBox2D GetUVRegion() const
+	{
+		return UVRegion;
+	}
+
+	/**
+	 * Set brush UV region
+	 *
+	 * @param InUVRegion When valid - overrides UV region specified in resource proxy
+	 */
+	void SetUVRegion(const FBox2D& InUVRegion)
+	{
+		UVRegion = InUVRegion;
+	}
+
 	/**
 	 * Compares this brush with another for equality.
 	 *
@@ -240,7 +261,8 @@ public:
 			&& Mirroring == Other.Mirroring
 			&& ResourceObject == Other.ResourceObject
 			&& ResourceName == Other.ResourceName
-			&& bIsDynamicallyLoaded == Other.bIsDynamicallyLoaded;
+			&& bIsDynamicallyLoaded == Other.bIsDynamicallyLoaded
+			&& UVRegion == Other.UVRegion;
 	}
 
 	/**
@@ -286,6 +308,13 @@ protected:
 	UPROPERTY()
 	bool bHasUObject_DEPRECATED;
 
+	/** 
+	 *  Optional UV region for an image
+	 *  When valid - overrides UV region specified in resource proxy
+	 */
+	UPROPERTY()
+	FBox2D UVRegion;
+
 protected:
 
 	/** 
diff --git a/Engine/Source/Runtime/SlateRHIRenderer/Private/SlateRHIResourceManager.cpp b/Engine/Source/Runtime/SlateRHIRenderer/Private/SlateRHIResourceManager.cpp
--- a/Engine/Source/Runtime/SlateRHIRenderer/Private/SlateRHIResourceManager.cpp
+++ b/Engine/Source/Runtime/SlateRHIRenderer/Private/SlateRHIResourceManager.cpp
@@ -23,7 +23,7 @@ TSharedPtr<FSlateDynamicTextureResource> FDynamicResourceMap::GetDynamicTextureR
 	return NativeTextureMap.FindRef( ResourceName );
 }
 
-TSharedPtr<FSlateUTextureResource> FDynamicResourceMap::GetUTextureResource( UTexture2D* TextureObject ) const
+TSharedPtr<FSlateUTextureResource> FDynamicResourceMap::GetUTextureResource( UTexture* TextureObject ) const
 {
 	if(TextureObject)
 	{
@@ -44,7 +44,7 @@ void FDynamicResourceMap::AddDynamicTextureResource( FName ResourceName, TShared
 	NativeTextureMap.Add( ResourceName, InResource );
 }
 
-void FDynamicResourceMap::AddUTextureResource( UTexture2D* TextureObject, TSharedRef<FSlateUTextureResource> InResource)
+void FDynamicResourceMap::AddUTextureResource( UTexture* TextureObject, TSharedRef<FSlateUTextureResource> InResource)
 {
 	if ( TextureObject )
 	{
@@ -69,7 +69,7 @@ void FDynamicResourceMap::RemoveDynamicTextureResource(FName ResourceName)
 	NativeTextureMap.Remove(ResourceName);
 }
 
-void FDynamicResourceMap::RemoveUTextureResource( UTexture2D* TextureObject )
+void FDynamicResourceMap::RemoveUTextureResource( UTexture* TextureObject )
 {
 	if(TextureObject)
 	{
@@ -515,7 +515,7 @@ TSharedPtr<FSlateDynamicTextureResource> FSlateRHIResourceManager::GetDynamicTex
 	return DynamicResourceMap.GetDynamicTextureResource( ResourceName );
 }
 
-TSharedPtr<FSlateUTextureResource> FSlateRHIResourceManager::MakeDynamicUTextureResource(UTexture2D* InTextureObject)
+TSharedPtr<FSlateUTextureResource> FSlateRHIResourceManager::MakeDynamicUTextureResource(UTexture* InTextureObject)
 {
 	// Generated texture resource
 	TSharedPtr<FSlateUTextureResource> TextureResource;
@@ -554,7 +554,7 @@ TSharedPtr<FSlateUTextureResource> FSlateRHIResourceManager::MakeDynamicUTexture
 		
 		}
 
-		TextureResource->Proxy->ActualSize = FIntPoint(InTextureObject->GetSizeX(), InTextureObject->GetSizeY());
+		TextureResource->Proxy->ActualSize = FIntPoint(InTextureObject->GetSurfaceWidth(), InTextureObject->GetSurfaceHeight());
 
 		checkSlow(!AccessedUTextures.Contains(InTextureObject));
 	}
@@ -579,7 +579,7 @@ FSlateShaderResourceProxy* FSlateRHIResourceManager::FindOrCreateDynamicTextureR
 	{
 		if ( UObject* ResourceObject = InBrush.GetResourceObject() )
 		{
-			if ( !ResourceObject->IsA<UTexture2D>() )
+			if ( !(ResourceObject->IsA<UTexture2D>() || ResourceObject->IsA<UTexture2DDynamic>()) )
 			{
 #if !(UE_BUILD_SHIPPING || UE_BUILD_TEST)
 				static TSet<UObject*> FailedTextures;
@@ -594,7 +594,7 @@ FSlateShaderResourceProxy* FSlateRHIResourceManager::FindOrCreateDynamicTextureR
 #endif
 			}
 
-			UTexture2D* TextureObject = Cast<UTexture2D>(ResourceObject);
+			UTexture* TextureObject = Cast<UTexture>(ResourceObject);
 
 			TSharedPtr<FSlateUTextureResource> TextureResource = DynamicResourceMap.GetUTextureResource(TextureObject);
 
@@ -609,7 +609,7 @@ FSlateShaderResourceProxy* FSlateRHIResourceManager::FindOrCreateDynamicTextureR
 
 			if ( TextureResource.IsValid() )
 			{
-				UTexture2D* Texture2DObject = TextureResource->TextureObject;
+				UTexture* Texture2DObject = TextureResource->TextureObject;
 				if ( Texture2DObject && !AccessedUTextures.Contains(Texture2DObject) && Texture2DObject->Resource )
 				{
 					// Set the texture rendering resource that should be used.  The UTexture resource could change at any time so we must do this each frame
@@ -726,7 +726,7 @@ void FSlateRHIResourceManager::ReleaseDynamicResource( const FSlateBrush& InBrus
 
 		if( ResourceObject && DynamicResourceMap.GetNumObjectResources() > 0 )
 		{
-			TSharedPtr<FSlateUTextureResource> TextureResource = DynamicResourceMap.GetUTextureResource(Cast<UTexture2D>(ResourceObject));
+			TSharedPtr<FSlateUTextureResource> TextureResource = DynamicResourceMap.GetUTextureResource(Cast<UTexture>(ResourceObject));
 
 			if(TextureResource.IsValid())
 			{
diff --git a/Engine/Source/Runtime/SlateRHIRenderer/Private/SlateRHIResourceManager.h b/Engine/Source/Runtime/SlateRHIRenderer/Private/SlateRHIResourceManager.h
--- a/Engine/Source/Runtime/SlateRHIRenderer/Private/SlateRHIResourceManager.h
+++ b/Engine/Source/Runtime/SlateRHIRenderer/Private/SlateRHIResourceManager.h
@@ -15,12 +15,12 @@ public:
 
 	TSharedPtr<FSlateDynamicTextureResource> GetDynamicTextureResource( FName ResourceName ) const;
 
-	TSharedPtr<FSlateUTextureResource> GetUTextureResource( UTexture2D* TextureObject ) const;
+	TSharedPtr<FSlateUTextureResource> GetUTextureResource( UTexture* TextureObject ) const;
 
 	TSharedPtr<FSlateMaterialResource> GetMaterialResource( const UMaterialInterface* Material ) const;
 
-	void AddUTextureResource( UTexture2D* TextureObject, TSharedRef<FSlateUTextureResource> InResource );
-	void RemoveUTextureResource( UTexture2D* TextureObject );
+	void AddUTextureResource( UTexture* TextureObject, TSharedRef<FSlateUTextureResource> InResource );
+	void RemoveUTextureResource( UTexture* TextureObject );
 
 	void AddDynamicTextureResource( FName ResourceName, TSharedRef<FSlateDynamicTextureResource> InResource);
 	void RemoveDynamicTextureResource( FName ResourceName );
@@ -45,7 +45,7 @@ private:
 private:
 	TMap<FName, TSharedPtr<FSlateDynamicTextureResource> > NativeTextureMap;
 	
-	typedef TMap<TWeakObjectPtr<UTexture2D>, TSharedPtr<FSlateUTextureResource> > TextureResourceMap;
+	typedef TMap<TWeakObjectPtr<UTexture>, TSharedPtr<FSlateUTextureResource> > TextureResourceMap;
 
 	/** Map of all texture resources */
 	TextureResourceMap TextureMap;
@@ -106,7 +106,7 @@ public:
 	 *
 	 * @param InTextureObject	The texture object to create the resource from
 	 */
-	TSharedPtr<FSlateUTextureResource> MakeDynamicUTextureResource( UTexture2D* InTextureObject);
+	TSharedPtr<FSlateUTextureResource> MakeDynamicUTextureResource( UTexture* InTextureObject);
 	
 	/**
 	 * Makes a dynamic texture resource and begins use of it
@@ -203,7 +203,7 @@ private:
 	/** Map of all active dynamic resources being used by brushes */
 	FDynamicResourceMap DynamicResourceMap;
 	/** Set of dynamic textures that are currently being accessed */
-	TSet< TWeakObjectPtr<UTexture2D> > AccessedUTextures;
+	TSet< TWeakObjectPtr<UTexture> > AccessedUTextures;
 	/** List of old utexture resources that are free to use as new resources */
 	TArray< TSharedPtr<FSlateUTextureResource> > UTextureFreeList;
 	/** List of old dynamic resources that are free to use as new resources */
diff --git a/Engine/Source/Runtime/SlateRHIRenderer/Private/SlateUTextureResource.cpp b/Engine/Source/Runtime/SlateRHIRenderer/Private/SlateUTextureResource.cpp
--- a/Engine/Source/Runtime/SlateRHIRenderer/Private/SlateUTextureResource.cpp
+++ b/Engine/Source/Runtime/SlateRHIRenderer/Private/SlateUTextureResource.cpp
@@ -5,13 +5,13 @@
 
 TSharedPtr<FSlateUTextureResource> FSlateUTextureResource::NullResource = MakeShareable( new FSlateUTextureResource(nullptr) );
 
-FSlateUTextureResource::FSlateUTextureResource(UTexture2D* InTexture)
+FSlateUTextureResource::FSlateUTextureResource(UTexture* InTexture)
 	: Proxy(new FSlateShaderResourceProxy)
 	, TextureObject(InTexture)
 {
 	if(TextureObject)
 	{
-		Proxy->ActualSize = FIntPoint(TextureObject->GetSizeX(), TextureObject->GetSizeY());
+		Proxy->ActualSize = FIntPoint(InTexture->GetSurfaceWidth(), InTexture->GetSurfaceHeight());
 		Proxy->Resource = this;
 	}
 }
@@ -43,12 +43,12 @@ void FSlateUTextureResource::UpdateRenderResource(FTexture* InFTexture)
 
 uint32 FSlateUTextureResource::GetWidth() const
 { 
-	return TextureObject->GetSizeX();
+	return TextureObject->GetSurfaceWidth();
 }
 
 uint32 FSlateUTextureResource::GetHeight() const 
 { 
-	return TextureObject->GetSizeY(); 
+	return TextureObject->GetSurfaceHeight(); 
 }
 
 ESlateShaderResource::Type FSlateUTextureResource::GetType() const
diff --git a/Engine/Source/Runtime/SlateRHIRenderer/Private/SlateUTextureResource.h b/Engine/Source/Runtime/SlateRHIRenderer/Private/SlateUTextureResource.h
--- a/Engine/Source/Runtime/SlateRHIRenderer/Private/SlateUTextureResource.h
+++ b/Engine/Source/Runtime/SlateRHIRenderer/Private/SlateUTextureResource.h
@@ -8,7 +8,7 @@ class FSlateUTextureResource : public TSlateTexture<FTexture2DRHIRef>
 public:
 	static TSharedPtr<FSlateUTextureResource> NullResource;
 
-	FSlateUTextureResource(UTexture2D* InTexture);
+	FSlateUTextureResource(UTexture* InTexture);
 	~FSlateUTextureResource();
 
 	/**
@@ -28,5 +28,5 @@ public:
 	FSlateShaderResourceProxy* Proxy;
 
 	/** Texture UObject.  Note: lifetime is managed externally */
-	UTexture2D* TextureObject;
+	UTexture* TextureObject;
 };
