Slate - Whenever the game viewport is torn down, we now immediately flush the accessed resource cache so that any dynamically created materials have an opportunity to be GCed when the game is being torn down to prevent a crash when the end play code does a check against non-collected objects still referencing the dying world.

diff --git a/Engine/Source/Runtime/Slate/Private/Framework/Application/SlateApplication.cpp b/Engine/Source/Runtime/Slate/Private/Framework/Application/SlateApplication.cpp
--- a/Engine/Source/Runtime/Slate/Private/Framework/Application/SlateApplication.cpp
+++ b/Engine/Source/Runtime/Slate/Private/Framework/Application/SlateApplication.cpp
@@ -2039,6 +2039,11 @@ void FSlateApplication::UnregisterGameViewport()
 		GameViewportWidget.Pin()->SetActive(false);
 	}
 	GameViewportWidget.Reset();
+
+	// Release any temporary material or texture resources we may have cached and are reporting to prevent
+	// GC on those resources.  If the game viewport is being unregistered, we need to flush these resources
+	// to allow for them to be GC'ed.
+	Renderer->ReleaseAccessedResources();
 }
 
 TSharedPtr<SViewport> FSlateApplication::GetGameViewport() const
diff --git a/Engine/Source/Runtime/SlateCore/Public/Rendering/SlateRenderer.h b/Engine/Source/Runtime/SlateCore/Public/Rendering/SlateRenderer.h
--- a/Engine/Source/Runtime/SlateCore/Public/Rendering/SlateRenderer.h
+++ b/Engine/Source/Runtime/SlateCore/Public/Rendering/SlateRenderer.h
@@ -187,6 +187,13 @@ public:
 	 */
 	virtual void InvalidateAllViewports() {}
 
+	/**
+	 * A renderer may need to keep a cache of accessed garbage collectible objects alive for the duration of their
+	 * usage.  During some operations like ending a game.  It becomes important to immediately release game related
+	 * resources.  This should flush any buffer holding onto those referenced objects.
+	 */
+	virtual void ReleaseAccessedResources() {}
+
 	/** 
 	 * Prepares the renderer to take a screenshot of the UI.  The Rect is portion of the rendered output
 	 * that will be stored into the TArray of FColors.
diff --git a/Engine/Source/Runtime/SlateRHIRenderer/Private/SlateRHIRenderer.cpp b/Engine/Source/Runtime/SlateRHIRenderer/Private/SlateRHIRenderer.cpp
--- a/Engine/Source/Runtime/SlateRHIRenderer/Private/SlateRHIRenderer.cpp
+++ b/Engine/Source/Runtime/SlateRHIRenderer/Private/SlateRHIRenderer.cpp
@@ -566,8 +566,7 @@ void FSlateRHIRenderer::DrawWindows_Private( FSlateDrawBuffer& WindowDrawBuffer
 		Policy.BeginDrawingWindows();
 	});
 
-	// Clear accessed UTexture and Material objects from the previous frame
-	ResourceManager->ClearAccessedResources();
+	ReleaseAccessedResources();
 
 	// Update texture atlases if needed
 	ResourceManager->UpdateTextureAtlases();
@@ -1208,10 +1207,16 @@ void FSlateRHIRenderer::InvalidateAllViewports()
 {
 	for( TMap< const SWindow*, FViewportInfo*>::TIterator It(WindowToViewportInfo); It; ++It )
 	{
-		It.Value()->ViewportRHI = NULL;
+		It.Value()->ViewportRHI = nullptr;
 	}
 }
 
+void FSlateRHIRenderer::ReleaseAccessedResources()
+{
+	// Clear accessed UTexture and Material objects from the previous frame
+	ResourceManager->ReleaseAccessedResources();
+}
+
 void FSlateRHIRenderer::RequestResize( const TSharedPtr<SWindow>& Window, uint32 NewWidth, uint32 NewHeight )
 {
 	check( IsThreadSafeForSlateRendering() );
diff --git a/Engine/Source/Runtime/SlateRHIRenderer/Private/SlateRHIRenderer.h b/Engine/Source/Runtime/SlateRHIRenderer/Private/SlateRHIRenderer.h
--- a/Engine/Source/Runtime/SlateRHIRenderer/Private/SlateRHIRenderer.h
+++ b/Engine/Source/Runtime/SlateRHIRenderer/Private/SlateRHIRenderer.h
@@ -187,6 +187,7 @@ public:
 	virtual FSlateUpdatableTexture* CreateUpdatableTexture(uint32 Width, uint32 Height) override;
 	virtual void ReleaseUpdatableTexture(FSlateUpdatableTexture* Texture) override;
 	virtual ISlateAtlasProvider* GetTextureAtlasProvider() override;
+	virtual void ReleaseAccessedResources() override;
 
 	/** Draws windows from a FSlateDrawBuffer on the render thread */
 	void DrawWindow_RenderThread(FRHICommandListImmediate& RHICmdList, const FSlateRHIRenderer::FViewportInfo& ViewportInfo, FSlateWindowElementList& WindowElementList, bool bLockToVsync, bool bClear);
diff --git a/Engine/Source/Runtime/SlateRHIRenderer/Private/SlateRHIResourceManager.cpp b/Engine/Source/Runtime/SlateRHIRenderer/Private/SlateRHIResourceManager.cpp
--- a/Engine/Source/Runtime/SlateRHIRenderer/Private/SlateRHIResourceManager.cpp
+++ b/Engine/Source/Runtime/SlateRHIRenderer/Private/SlateRHIResourceManager.cpp
@@ -798,7 +798,7 @@ void FSlateRHIResourceManager::LoadStyleResources( const ISlateStyle& Style )
 	CreateTextures( Resources );
 }
 
-void FSlateRHIResourceManager::ClearAccessedResources()
+void FSlateRHIResourceManager::ReleaseAccessedResources()
 {
 	AccessedUTextures.Reset();
 	AccessedMaterials.Reset();
diff --git a/Engine/Source/Runtime/SlateRHIRenderer/Private/SlateRHIResourceManager.h b/Engine/Source/Runtime/SlateRHIRenderer/Private/SlateRHIResourceManager.h
--- a/Engine/Source/Runtime/SlateRHIRenderer/Private/SlateRHIResourceManager.h
+++ b/Engine/Source/Runtime/SlateRHIRenderer/Private/SlateRHIResourceManager.h
@@ -91,7 +91,7 @@ public:
 	 * The accessed textures is used to determine which textures need be updated on the render thread
 	 * so they can be used by slate
 	 */
-	void ClearAccessedResources();
+	void ReleaseAccessedResources();
 
 	/**
 	 * Updates texture atlases if needed
