Add a resolve to properly support fastclears.

diff --git a/Engine/Source/Runtime/SlateRHIRenderer/Private/Slate3DRenderer.cpp b/Engine/Source/Runtime/SlateRHIRenderer/Private/Slate3DRenderer.cpp
--- a/Engine/Source/Runtime/SlateRHIRenderer/Private/Slate3DRenderer.cpp
+++ b/Engine/Source/Runtime/SlateRHIRenderer/Private/Slate3DRenderer.cpp
@@ -93,7 +93,8 @@ void FSlate3DRenderer::DrawWindowToTarget_RenderThread( FRHICommandListImmediate
 	FTextureRenderTarget2DResource* RenderTargetResource = static_cast<FTextureRenderTarget2DResource*>( RenderTarget->GetRenderTargetResource() );
 	
 	// Set render target and clear.
-	FRHIRenderTargetView ColorRTV(RenderTargetResource->GetTextureRHI());
+	FTexture2DRHIRef RTResource = RenderTargetResource->GetTextureRHI();
+	FRHIRenderTargetView ColorRTV(RTResource);
 	FRHISetRenderTargetsInfo Info(1, &ColorRTV, FTextureRHIParamRef());
 	Info.bClearColor = true;
 	ensure(ColorRTV.Texture->GetClearColor() == RenderTarget->ClearColor);
@@ -121,4 +122,6 @@ void FSlate3DRenderer::DrawWindowToTarget_RenderThread( FRHICommandListImmediate
 	InDrawBuffer.Unlock();
 
 	RenderTargetPolicy->EndDrawingWindows();
+
+	RHICmdList.CopyToResolveTarget(RenderTargetResource->GetTextureRHI(), RTResource, true, FResolveParams());
 }
\ No newline at end of file
