fix for 3d widgets not being cleared properly on Metal devices

diff --git a/Engine/Source/Runtime/SlateRHIRenderer/Private/Slate3DRenderer.cpp b/Engine/Source/Runtime/SlateRHIRenderer/Private/Slate3DRenderer.cpp
--- a/Engine/Source/Runtime/SlateRHIRenderer/Private/Slate3DRenderer.cpp
+++ b/Engine/Source/Runtime/SlateRHIRenderer/Private/Slate3DRenderer.cpp
@@ -95,6 +95,7 @@ void FSlate3DRenderer::DrawWindowToTarget_RenderThread( FRHICommandListImmediate
 	// Set render target and clear.
 	FTexture2DRHIRef RTResource = RenderTargetResource->GetTextureRHI();
 	FRHIRenderTargetView ColorRTV(RTResource);
+	ColorRTV.LoadAction = ERenderTargetLoadAction::EClear;
 	FRHISetRenderTargetsInfo Info(1, &ColorRTV, FTextureRHIParamRef());
 	Info.bClearColor = true;
 	ensure(ColorRTV.Texture->GetClearColor() == RenderTarget->ClearColor);
