Merged from Orion CL 2660007
Remove redundant present time that bloats Slate rendering time (present time is not indicative of Slate being slow).  It exists in the RHI stat and in stat scenerendering

diff --git a/Engine/Source/Runtime/SlateRHIRenderer/Private/SlateRHIRenderer.cpp b/Engine/Source/Runtime/SlateRHIRenderer/Private/SlateRHIRenderer.cpp
--- a/Engine/Source/Runtime/SlateRHIRenderer/Private/SlateRHIRenderer.cpp
+++ b/Engine/Source/Runtime/SlateRHIRenderer/Private/SlateRHIRenderer.cpp
@@ -470,12 +470,8 @@ void FSlateRHIRenderer::DrawWindow_RenderThread(FRHICommandListImmediate& RHICmd
 
 	// Calculate renderthread time (excluding idle time).	
 	uint32 StartTime		= FPlatformTime::Cycles();
-		
-	// Note - We do not include present time in the slate render thread stat
-	{
-		SCOPE_CYCLE_COUNTER(STAT_SlatePresentRTTime);
-		RHICmdList.EndDrawingViewport(ViewportInfo.ViewportRHI, true, bLockToVsync);
-	}
+
+	RHICmdList.EndDrawingViewport(ViewportInfo.ViewportRHI, true, bLockToVsync);
 
 	if (bNeedCallFinishFrameForStereo)
 	{
