UMG/Slate - Continuing to itterate on performance and improving the invalidation panel.  Enhanced the rendering pipeline to allow cached buffers to be passed along with an offset so that we could do a final matrix transform on it if need be.  Fixing a rather large memory leak on the invalidation panels never being able to re-use old element lists because they were always being used by the render thread.  Adding more stats.  The PS4 BeginNamedEvent calls now properly respect the color being passed as an argument and will display it in Razor when in arguement coloring mode.  Widgets now know what class created them, in all but shipping.  Very useful for logging stat/debugging information in a non-editor build.

diff --git a/Engine/Source/Runtime/Core/Private/Math/Color.cpp b/Engine/Source/Runtime/Core/Private/Math/Color.cpp
--- a/Engine/Source/Runtime/Core/Private/Math/Color.cpp
+++ b/Engine/Source/Runtime/Core/Private/Math/Color.cpp
@@ -26,6 +26,7 @@ const FColor FColor::Yellow(255,255,0);
 const FColor FColor::Cyan(0,255,255);
 const FColor FColor::Magenta(255,0,255);
 const FColor FColor::Orange(243, 156, 18);
+const FColor FColor::Purple(169, 7, 228);
 const FColor FColor::Turquoise(26, 188, 156);
 const FColor FColor::Silver(189, 195, 199);
 const FColor FColor::Emerald(46, 204, 113);
diff --git a/Engine/Source/Runtime/Core/Public/Math/Color.h b/Engine/Source/Runtime/Core/Public/Math/Color.h
--- a/Engine/Source/Runtime/Core/Public/Math/Color.h
+++ b/Engine/Source/Runtime/Core/Public/Math/Color.h
@@ -562,6 +562,38 @@ public:
 		return bSuccessful;
 	}
 
+	/**
+	 * Gets the color in a packed uint32 format packed in the order ARGB.
+	 */
+	FORCEINLINE uint32 ToPackedARGB() const
+	{
+		return ( A << 24 ) | ( R << 16 ) | ( G << 8 ) | ( B << 0 );
+	}
+
+	/**
+	 * Gets the color in a packed uint32 format packed in the order ABGR.
+	 */
+	FORCEINLINE uint32 ToPackedABGR() const
+	{
+		return ( A << 24 ) | ( B << 16 ) | ( G << 8 ) | ( R << 0 );
+	}
+
+	/**
+	 * Gets the color in a packed uint32 format packed in the order RGBA.
+	 */
+	FORCEINLINE uint32 ToPackedRGBA() const
+	{
+		return ( R << 24 ) | ( G << 16 ) | ( B << 8 ) | ( A << 0 );
+	}
+
+	/**
+	 * Gets the color in a packed uint32 format packed in the order BGRA.
+	 */
+	FORCEINLINE uint32 ToPackedBGRA() const
+	{
+		return ( B << 24 ) | ( G << 16 ) | ( R << 8 ) | ( A << 0 );
+	}
+
 	/** Some pre-inited colors, useful for debug code */
 	static CORE_API const FColor White;
 	static CORE_API const FColor Black;
@@ -573,6 +605,7 @@ public:
 	static CORE_API const FColor Cyan;
 	static CORE_API const FColor Magenta;
 	static CORE_API const FColor Orange;
+	static CORE_API const FColor Purple;
 	static CORE_API const FColor Turquoise;
 	static CORE_API const FColor Silver;
 	static CORE_API const FColor Emerald;
diff --git a/Engine/Source/Runtime/Engine/Private/Slate/SGameLayerManager.cpp b/Engine/Source/Runtime/Engine/Private/Slate/SGameLayerManager.cpp
--- a/Engine/Source/Runtime/Engine/Private/Slate/SGameLayerManager.cpp
+++ b/Engine/Source/Runtime/Engine/Private/Slate/SGameLayerManager.cpp
@@ -167,9 +167,9 @@ void SGameLayerManager::Tick(const FGeometry& AllottedGeometry, const double InC
 
 int32 SGameLayerManager::OnPaint(const FPaintArgs& Args, const FGeometry& AllottedGeometry, const FSlateRect& MyClippingRect, FSlateWindowElementList& OutDrawElements, int32 LayerId, const FWidgetStyle& InWidgetStyle, bool bParentEnabled) const
 {
-	//FPlatformMisc::BeginNamedEvent(FColor::Green, "Game Layer Manager");
-	int32 ResultLayer = SCompoundWidget::OnPaint(Args, AllottedGeometry, MyClippingRect, OutDrawElements, LayerId, InWidgetStyle, bParentEnabled);
-	//FPlatformMisc::EndNamedEvent();
+	FPlatformMisc::BeginNamedEvent(FColor::Green, "Paint: Game UI");
+	const int32 ResultLayer = SCompoundWidget::OnPaint(Args, AllottedGeometry, MyClippingRect, OutDrawElements, LayerId, InWidgetStyle, bParentEnabled);
+	FPlatformMisc::EndNamedEvent();
 	return ResultLayer;
 }
 
diff --git a/Engine/Source/Runtime/Slate/Private/Widgets/Layout/SBorder.cpp b/Engine/Source/Runtime/Slate/Private/Widgets/Layout/SBorder.cpp
--- a/Engine/Source/Runtime/Slate/Private/Widgets/Layout/SBorder.cpp
+++ b/Engine/Source/Runtime/Slate/Private/Widgets/Layout/SBorder.cpp
@@ -63,11 +63,10 @@ void SBorder::ClearContent()
 
 int32 SBorder::OnPaint( const FPaintArgs& Args, const FGeometry& AllottedGeometry, const FSlateRect& MyClippingRect, FSlateWindowElementList& OutDrawElements, int32 LayerId, const FWidgetStyle& InWidgetStyle, bool bParentEnabled ) const
 {
-	bool bEnabled = false;
 	const FSlateBrush* BrushResource = BorderImage.Get();
 		
-	bEnabled = ShouldBeEnabled( bParentEnabled );
-	bool bShowDisabledEffect = ShowDisabledEffect.Get();
+	const bool bEnabled = ShouldBeEnabled(bParentEnabled);
+	const bool bShowDisabledEffect = ShowDisabledEffect.Get();
 	ESlateDrawEffect::Type DrawEffects = bShowDisabledEffect && !bEnabled ? ESlateDrawEffect::DisabledEffect : ESlateDrawEffect::None;
 
 	if ( BrushResource && BrushResource->DrawAs != ESlateBrushDrawType::NoDrawType )
diff --git a/Engine/Source/Runtime/Slate/Private/Widgets/SInvalidationPanel.cpp b/Engine/Source/Runtime/Slate/Private/Widgets/SInvalidationPanel.cpp
--- a/Engine/Source/Runtime/Slate/Private/Widgets/SInvalidationPanel.cpp
+++ b/Engine/Source/Runtime/Slate/Private/Widgets/SInvalidationPanel.cpp
@@ -4,6 +4,7 @@
 #include "LayoutUtils.h"
 #include "SInvalidationPanel.h"
 #include "WidgetCaching.h"
+#include "ReflectionMetaData.h"
 
 //DECLARE_CYCLE_STAT(TEXT("Invalidation Time"), STAT_InvalidationTime, STATGROUP_Slate);
 DECLARE_DWORD_COUNTER_STAT(TEXT("Num Cached Elements"), STAT_SlateNumCachedElements, STATGROUP_Slate);
@@ -62,6 +63,9 @@ SInvalidationPanel::~SInvalidationPanel()
 	{
 		delete NodePool[i];
 	}
+
+	TSharedPtr<FSlateRenderer> Renderer = FSlateApplicationBase::Get().GetRenderer();
+	Renderer->ReleaseCachingResourcesFor(this);
 }
 
 bool SInvalidationPanel::GetCanCache() const
@@ -228,6 +232,14 @@ int32 SInvalidationPanel::OnPaint( const FPaintArgs& Args, const FGeometry& Allo
 {
 	if ( GetCanCache() )
 	{
+		// If our clip rect changes size, we've definitely got to invalidate.
+		const FVector2D ClipRectSize = MyClippingRect.GetSize();
+		if ( ClipRectSize != LastClipRectSize )
+		{
+			bNeedsCaching = true;
+			LastClipRectSize = ClipRectSize;
+		}
+
 		//SCOPE_CYCLE_COUNTER(STAT_InvalidationTime);
 
 		//FPlatformMisc::BeginNamedEvent(FColor::Magenta, "Slate::InvalidationPanel::Paint");
@@ -291,7 +303,7 @@ int32 SInvalidationPanel::OnPaint( const FPaintArgs& Args, const FGeometry& Allo
 			}
 
 #if WITH_ENGINE
-			CachedRenderData = CachedWindowElements->CacheRenderData();
+			CachedRenderData = CachedWindowElements->CacheRenderData(this);
 #endif
 
 			LastHitTestIndex = Args.GetLastHitTestIndex();
@@ -322,16 +334,10 @@ int32 SInvalidationPanel::OnPaint( const FPaintArgs& Args, const FGeometry& Allo
 		if ( bCacheRelativeTransforms )
 		{
 			FVector2D NewAbsoluteDeltaPosition = AllottedGeometry.Position - CachedAbsolutePosition;
-			FVector2D RelativeDeltaPosition = NewAbsoluteDeltaPosition - AbsoluteDeltaPosition;
 			AbsoluteDeltaPosition = NewAbsoluteDeltaPosition;
 
 #if WITH_ENGINE
-			if ( bWasCachingNeeded == false && RelativeDeltaPosition.IsZero() == false )
-			{
-				CachedWindowElements->UpdateCacheRenderData(RelativeDeltaPosition);
-			}
-
-			FSlateDrawElement::MakeCachedBuffer(OutDrawElements, LayerId, CachedRenderData);
+			FSlateDrawElement::MakeCachedBuffer(OutDrawElements, LayerId, CachedRenderData, AbsoluteDeltaPosition * AllottedGeometry.Scale);
 #else
 			const TArray<FSlateDrawElement>& CachedElements = CachedWindowElements->GetDrawElements();
 			const int32 CachedElementCount = CachedElements.Num();
@@ -350,23 +356,25 @@ int32 SInvalidationPanel::OnPaint( const FPaintArgs& Args, const FGeometry& Allo
 		else
 		{
 #if WITH_ENGINE
-			FSlateDrawElement::MakeCachedBuffer(OutDrawElements, LayerId, CachedRenderData);
+			FSlateDrawElement::MakeCachedBuffer(OutDrawElements, LayerId, CachedRenderData, FVector2D(0, 0));
 #else
 			OutDrawElements.AppendDrawElements(CachedWindowElements->GetDrawElements());
 #endif
 		}
 
-		//FPlatformMisc::BeginNamedEvent(FColor::Orange, "Slate::VolatilePainting");
 		// Paint the volatile elements
 		if ( CachedWindowElements.IsValid() )
 		{
+			//FPlatformMisc::BeginNamedEvent(FColor::Red, *FReflectionMetaData::GetWidgetDebugInfo(this));
+
 			const TArray<TSharedPtr<FSlateWindowElementList::FVolatilePaint>>& VolatileElements = CachedWindowElements->GetVolatileElements();
 			INC_DWORD_STAT_BY(STAT_SlateNumVolatileWidgets, VolatileElements.Num());
 
 			// TODO Offset? AbsoluteDeltaPosition
 			OutMaxChildLayer = FMath::Max(OutMaxChildLayer, CachedWindowElements->PaintVolatile(OutDrawElements));
+
+			//FPlatformMisc::EndNamedEvent();
 		}
-		//FPlatformMisc::EndNamedEvent();
 
 #if !(UE_BUILD_SHIPPING || UE_BUILD_TEST)
 
diff --git a/Engine/Source/Runtime/Slate/Private/Widgets/SViewport.cpp b/Engine/Source/Runtime/Slate/Private/Widgets/SViewport.cpp
--- a/Engine/Source/Runtime/Slate/Private/Widgets/SViewport.cpp
+++ b/Engine/Source/Runtime/Slate/Private/Widgets/SViewport.cpp
@@ -3,7 +3,7 @@
 #include "SlatePrivatePCH.h"
 #include "HittestGrid.h"
 
-DECLARE_CYCLE_STAT(TEXT("Viewport Update Time"), STAT_ViewportUpdateTime, STATGROUP_Slate);
+DECLARE_CYCLE_STAT(TEXT("Game UI Tick/Paint"), STAT_ViewportUpdateTime, STATGROUP_Slate);
 
 /* SViewport constructors
  *****************************************************************************/
diff --git a/Engine/Source/Runtime/Slate/Private/Widgets/Text/STextBlock.cpp b/Engine/Source/Runtime/Slate/Private/Widgets/Text/STextBlock.cpp
--- a/Engine/Source/Runtime/Slate/Private/Widgets/Text/STextBlock.cpp
+++ b/Engine/Source/Runtime/Slate/Private/Widgets/Text/STextBlock.cpp
@@ -103,11 +103,27 @@ void STextBlock::SetText( const TAttribute< FText >& InText )
 void STextBlock::SetText( const FText& InText )
 {
 	SCOPE_CYCLE_COUNTER(Stat_SlateTextBlockSetText);
-	if ( BoundText.IsBound() || BoundText.Get().ToString() != InText.ToString() )
+
+	if ( !BoundText.IsBound() )
 	{
-		BoundText = InText;
-		Invalidate(EInvalidateWidget::LayoutAndVolatility);
+		const FString& OldString = BoundText.Get().ToString();
+		const FString& NewString = InText.ToString();
+		const int32 OldLength = OldString.Len();
+		const int32 NewLength = NewString.Len();
+
+		// We only perform this optimization if the text we're checking is smaller 
+		// than 30 characters otherwise we may be comparing pages of text.
+		if ( OldLength == NewLength && OldLength <= 30 )
+		{
+			if ( OldString == NewString )
+			{
+				return;
+			}
+		}
 	}
+
+	BoundText = InText;
+	Invalidate(EInvalidateWidget::LayoutAndVolatility);
 }
 
 int32 STextBlock::OnPaint( const FPaintArgs& Args, const FGeometry& AllottedGeometry, const FSlateRect& MyClippingRect, FSlateWindowElementList& OutDrawElements, int32 LayerId, const FWidgetStyle& InWidgetStyle, bool bParentEnabled ) const
@@ -116,9 +132,13 @@ int32 STextBlock::OnPaint( const FPaintArgs& Args, const FGeometry& AllottedGeom
 
 #if WITH_FANCY_TEXT
 
+	//FPlatformMisc::BeginNamedEvent(FColor::Orange, "STextBlock");
+
 	// OnPaint will also update the text layout cache if required
 	LayerId = TextLayoutCache->OnPaint(Args, AllottedGeometry, MyClippingRect, OutDrawElements, LayerId, InWidgetStyle, ShouldBeEnabled(bParentEnabled));
 
+	//FPlatformMisc::EndNamedEvent();
+
 #else//WITH_FANCY_TEXT
 
 	// Simple stubbed version for when WITH_FANCY_TEXT is disabled (which should only ever be for testing purposes)
diff --git a/Engine/Source/Runtime/Slate/Public/Widgets/Layout/SConstraintCanvas.cpp b/Engine/Source/Runtime/Slate/Public/Widgets/Layout/SConstraintCanvas.cpp
--- a/Engine/Source/Runtime/Slate/Public/Widgets/Layout/SConstraintCanvas.cpp
+++ b/Engine/Source/Runtime/Slate/Public/Widgets/Layout/SConstraintCanvas.cpp
@@ -163,6 +163,8 @@ void SConstraintCanvas::OnArrangeChildren( const FGeometry& AllottedGeometry, FA
 
 int32 SConstraintCanvas::OnPaint( const FPaintArgs& Args, const FGeometry& AllottedGeometry, const FSlateRect& MyClippingRect, FSlateWindowElementList& OutDrawElements, int32 LayerId, const FWidgetStyle& InWidgetStyle, bool bParentEnabled ) const
 {
+	//FPlatformMisc::BeginNamedEvent(FColor::Orange, "SConstraintCanvas");
+
 	FArrangedChildren ArrangedChildren(EVisibility::Visible);
 	ArrangeChildren(AllottedGeometry, ArrangedChildren);
 
@@ -185,6 +187,8 @@ int32 SConstraintCanvas::OnPaint( const FPaintArgs& Args, const FGeometry& Allot
 		}
 	}
 
+	//FPlatformMisc::EndNamedEvent();
+
 	return MaxLayerId;
 }
 
diff --git a/Engine/Source/Runtime/Slate/Public/Widgets/SInvalidationPanel.h b/Engine/Source/Runtime/Slate/Public/Widgets/SInvalidationPanel.h
--- a/Engine/Source/Runtime/Slate/Public/Widgets/SInvalidationPanel.h
+++ b/Engine/Source/Runtime/Slate/Public/Widgets/SInvalidationPanel.h
@@ -74,6 +74,7 @@ private:
 	mutable TArray< FCachedWidgetNode* > NodePool;
 	mutable int32 LastUsedCachedNodeIndex;
 	mutable int32 LastHitTestIndex;
+	mutable FVector2D LastClipRectSize;
 
 	mutable int32 CachedMaxChildLayer;
 	mutable bool bNeedsCaching;
diff --git a/Engine/Source/Runtime/SlateCore/Private/Rendering/DrawElements.cpp b/Engine/Source/Runtime/SlateCore/Private/Rendering/DrawElements.cpp
--- a/Engine/Source/Runtime/SlateCore/Private/Rendering/DrawElements.cpp
+++ b/Engine/Source/Runtime/SlateCore/Private/Rendering/DrawElements.cpp
@@ -162,14 +162,16 @@ void FSlateDrawElement::MakeCustom( FSlateWindowElementList& ElementList, uint32
 	DrawElt.DataPayload.SetCustomDrawerPayloadProperties( CustomDrawer );
 }
 
-void FSlateDrawElement::MakeCachedBuffer(FSlateWindowElementList& ElementList, uint32 InLayer, TSharedPtr<FSlateRenderDataHandle, ESPMode::ThreadSafe>& CachedRenderDataHandle)
+void FSlateDrawElement::MakeCachedBuffer(FSlateWindowElementList& ElementList, uint32 InLayer, TSharedPtr<FSlateRenderDataHandle, ESPMode::ThreadSafe>& CachedRenderDataHandle, const FVector2D& Offset)
 {
 	SCOPE_CYCLE_COUNTER(STAT_SlateDrawElementMakeTime)
 	FSlateDrawElement& DrawElt = ElementList.AddUninitialized();
 	DrawElt.Init(InLayer, FPaintGeometry(), FSlateRect(1, 1, 1, 1), ESlateDrawEffect::None);
 	DrawElt.RenderTransform = FSlateRenderTransform();
 	DrawElt.ElementType = ET_CachedBuffer;
-	DrawElt.DataPayload.SetCachedBufferPayloadProperties(CachedRenderDataHandle.Get());
+	DrawElt.DataPayload.SetCachedBufferPayloadProperties(CachedRenderDataHandle.Get(), Offset);
+
+	ElementList.BeginUsingCachedBuffer(CachedRenderDataHandle);
 }
 
 void FSlateDrawElement::MakeLayer(FSlateWindowElementList& ElementList, uint32 InLayer, TSharedPtr<FSlateDrawLayerHandle, ESPMode::ThreadSafe>& DrawLayerHandle)
@@ -213,6 +215,7 @@ FVector2D FSlateDrawElement::GetRotationPoint(const FPaintGeometry& PaintGeometr
 void FSlateBatchData::Reset()
 {
 	RenderBatches.Reset();
+	DynamicOffset = FVector2D(0, 0);
 	
 	// note: LayerToElementBatches is not reset here as the same layers are 
 	// more than likely reused and we can save memory allocations by not resetting the map every frame
@@ -299,28 +302,6 @@ void FSlateBatchData::FillVertexAndIndexBuffer( uint8* VertexBuffer, uint8* Inde
 	}
 }
 
-void FSlateBatchData::UpdateRenderBatches(FVector2D PositionOffset)
-{
-	checkSlow(IsInRenderingThread());
-
-	for ( const FSlateRenderBatch& Batch : RenderBatches )
-	{
-		if ( Batch.VertexArrayIndex != INDEX_NONE && Batch.IndexArrayIndex != INDEX_NONE )
-		{
-			TArray<FSlateVertex>& Vertices = BatchVertexArrays[Batch.VertexArrayIndex];
-
-			int32 Count = Vertices.Num();
-			for ( int32 VertexIndex = 0; VertexIndex < Count; VertexIndex++ )
-			{
-				FSlateVertex& Vertex = Vertices[VertexIndex];
-				Vertex.Position[0] += PositionOffset.X;
-				Vertex.Position[1] += PositionOffset.Y;
-				Vertex.ClipRect.TopLeft += PositionOffset;
-			}
-		}
-	}
-}
-
 void FSlateBatchData::CreateRenderBatches(FElementBatchMap& LayerToElementBatches)
 {
 	checkSlow( IsInRenderingThread() );
@@ -337,6 +318,15 @@ void FSlateBatchData::CreateRenderBatches(FElementBatchMap& LayerToElementBatche
 	}
 }
 
+void FSlateBatchData::AddRenderBatch(uint32 InLayer, const FSlateElementBatch& InElementBatch, int32 InNumVertices, int32 InNumIndices, int32 InVertexOffset, int32 InIndexOffset)
+{
+	NumBatchedVertices += InNumVertices;
+	NumBatchedIndices += InNumIndices;
+
+	const int32 Index = RenderBatches.Add(FSlateRenderBatch(InLayer, InElementBatch, RenderDataHandle, InNumVertices, InNumIndices, InVertexOffset, InIndexOffset));
+	RenderBatches[Index].DynamicOffset = DynamicOffset;
+}
+
 void FSlateBatchData::Merge(FElementBatchMap& InLayerToElementBatches, uint32& VertexOffset, uint32& IndexOffset)
 {
 	InLayerToElementBatches.KeySort(TLess<uint32>());
@@ -366,6 +356,8 @@ void FSlateBatchData::Merge(FElementBatchMap& InLayerToElementBatches, uint32& V
 					{
 						if ( FSlateRenderDataHandle* RenderHandle = ElementBatch.GetCachedRenderHandle().Get() )
 						{
+							DynamicOffset += ElementBatch.GetCachedRenderDataOffset();
+
 							if ( TArray<FSlateRenderBatch>* ForeignBatches = RenderHandle->GetRenderBatches() )
 							{
 								TArray<FSlateRenderBatch>& ForeignBatchesRef = *ForeignBatches;
@@ -383,12 +375,14 @@ void FSlateBatchData::Merge(FElementBatchMap& InLayerToElementBatches, uint32& V
 									}
 									else
 									{
-										RenderBatches.Add(ForeignBatchesRef[i]);
+										const int32 Index = RenderBatches.Add(ForeignBatchesRef[i]);
+										RenderBatches[Index].DynamicOffset = DynamicOffset;
 									}
-									//RenderBatches.Append(*RenderHandle->GetRenderBatches());
 								}
 							}
 
+							DynamicOffset -= ElementBatch.GetCachedRenderDataOffset();
+
 							continue;
 						}
 					}
@@ -429,11 +423,6 @@ void FSlateBatchData::Merge(FElementBatchMap& InLayerToElementBatches, uint32& V
 	}
 }
 
-void FSlateBatchData::SortRenderBatches()
-{
-	//RenderBatches.StableSort();
-}
-
 
 FSlateWindowElementList::FDeferredPaint::FDeferredPaint( const TSharedRef<const SWidget>& InWidgetToPaint, const FPaintArgs& InArgs, const FGeometry InAllottedGeometry, const FSlateRect InMyClippingRect, const FWidgetStyle& InWidgetStyle, bool InParentEnabled )
 : WidgetToPaintPtr( InWidgetToPaint )
@@ -490,17 +479,16 @@ int32 FSlateWindowElementList::FVolatilePaint::ExecutePaint(FSlateWindowElementL
 	TSharedPtr<const SWidget> WidgetToPaint = WidgetToPaintPtr.Pin();
 	if ( WidgetToPaint.IsValid() )
 	{
-		//FPlatformMisc::BeginNamedEvent(FColor::Orange, "Overhead");
-		//FPlatformMisc::BeginNamedEvent(FColor::Orange, *FReflectionMetaData::GetWidgetDebugInfo(WidgetToPaint));
+		//FPlatformMisc::BeginNamedEvent(FColor::Red, *FReflectionMetaData::GetWidgetDebugInfo(WidgetToPaint));
 
 		// Have to run a slate pre-pass for all volatile elements, some widgets cache information like 
 		// the STextBlock.  This may be all kinds of terrible an idea to do during paint.
 		SWidget* MutableWidget = const_cast<SWidget*>( WidgetToPaint.Get() );
+
 		MutableWidget->SlatePrepass(AllottedGeometry.Scale);
 
 		const int32 NewLayer = WidgetToPaint->Paint(Args, AllottedGeometry, MyClippingRect, OutDrawElements, LayerId, WidgetStyle, bParentEnabled);
 		
-		//FPlatformMisc::EndNamedEvent();
 		//FPlatformMisc::EndNamedEvent();
 
 		return NewLayer;
@@ -570,9 +558,11 @@ void FSlateWindowElementList::EndLogicalLayer()
 	DrawStack.Pop();
 }
 
-FSlateRenderDataHandle::FSlateRenderDataHandle(FSlateRenderer* InRenderer)
-	: Renderer(InRenderer)
+FSlateRenderDataHandle::FSlateRenderDataHandle(const ILayoutCache* InCacher, FSlateRenderer* InRenderer)
+	: Cacher(InCacher)
+	, Renderer(InRenderer)
 	, RenderBatches(nullptr)
+	, UsageCount(0)
 {
 }
 
@@ -590,28 +580,19 @@ void FSlateRenderDataHandle::Disconnect()
 	RenderBatches = nullptr;
 }
 
-TSharedRef<FSlateRenderDataHandle, ESPMode::ThreadSafe> FSlateWindowElementList::CacheRenderData()
+TSharedRef<FSlateRenderDataHandle, ESPMode::ThreadSafe> FSlateWindowElementList::CacheRenderData(const ILayoutCache* Cacher)
 {
 	// Don't attempt to use this slate window element list if the cache is still being used.
 	checkSlow(!IsCachedRenderDataInUse());
 
 	TSharedPtr<FSlateRenderer> Renderer = FSlateApplicationBase::Get().GetRenderer();
 
-	TSharedRef<FSlateRenderDataHandle, ESPMode::ThreadSafe> CachedRenderDataHandleRef = Renderer->CacheElementRenderData(*this);
+	TSharedRef<FSlateRenderDataHandle, ESPMode::ThreadSafe> CachedRenderDataHandleRef = Renderer->CacheElementRenderData(Cacher, *this);
 	CachedRenderDataHandle = CachedRenderDataHandleRef;
 
 	return CachedRenderDataHandleRef;
 }
 
-void FSlateWindowElementList::UpdateCacheRenderData(FVector2D PositionOffset)
-{
-	// This should only be done if we're using cached render data
-	checkSlow(IsCachedRenderDataInUse());
-
-	TSharedPtr<FSlateRenderer> Renderer = FSlateApplicationBase::Get().GetRenderer();
-	Renderer->UpdateElementRenderData(*this, PositionOffset);
-}
-
 void FSlateWindowElementList::PreDraw_RenderThread()
 {
 	check(IsInRenderingThread());
@@ -631,6 +612,13 @@ void FSlateWindowElementList::PostDraw_RenderThread()
 	{
 		Entry.Key->BatchMap = nullptr;
 	}
+
+	for ( TSharedPtr<FSlateRenderDataHandle, ESPMode::ThreadSafe>& Handle : CachedRenderHandlesInUse )
+	{
+		Handle->EndUsing();
+	}
+
+	CachedRenderHandlesInUse.Reset();
 }
 
 void FSlateWindowElementList::ResetBuffers()
diff --git a/Engine/Source/Runtime/SlateCore/Private/Rendering/ElementBatcher.cpp b/Engine/Source/Runtime/SlateCore/Private/Rendering/ElementBatcher.cpp
--- a/Engine/Source/Runtime/SlateCore/Private/Rendering/ElementBatcher.cpp
+++ b/Engine/Source/Runtime/SlateCore/Private/Rendering/ElementBatcher.cpp
@@ -1578,7 +1578,7 @@ void FSlateElementBatcher::AddCachedBuffer(const FSlateDrawElement& DrawElement)
 
 		// Custom elements are not batched together
 		TSharedPtr< FSlateRenderDataHandle, ESPMode::ThreadSafe > RenderData = InPayload.CachedRenderData->AsShared();
-		ElementBatches->Add(FSlateElementBatch(RenderData, DrawElement.GetScissorRect()));
+		ElementBatches->Add(FSlateElementBatch(RenderData, InPayload.CachedRenderDataOffset, DrawElement.GetScissorRect()));
 	}
 }
 
diff --git a/Engine/Source/Runtime/SlateCore/Private/Rendering/SlateRenderer.cpp b/Engine/Source/Runtime/SlateCore/Private/Rendering/SlateRenderer.cpp
--- a/Engine/Source/Runtime/SlateCore/Private/Rendering/SlateRenderer.cpp
+++ b/Engine/Source/Runtime/SlateCore/Private/Rendering/SlateRenderer.cpp
@@ -55,14 +55,14 @@ ISlateAtlasProvider* FSlateRenderer::GetFontAtlasProvider()
 	return nullptr;
 }
 
-TSharedRef<FSlateRenderDataHandle, ESPMode::ThreadSafe> FSlateRenderer::CacheElementRenderData(FSlateWindowElementList& ElementList)
+TSharedRef<FSlateRenderDataHandle, ESPMode::ThreadSafe> FSlateRenderer::CacheElementRenderData(const ILayoutCache* Cacher, FSlateWindowElementList& ElementList)
 {
-	return MakeShareable(new FSlateRenderDataHandle(this));
+	return MakeShareable(new FSlateRenderDataHandle(Cacher, this));
 }
 
-void FSlateRenderer::UpdateElementRenderData(FSlateWindowElementList& ElementList, FVector2D PositionOffset)
+void FSlateRenderer::ReleaseCachingResourcesFor(const ILayoutCache* Cacher)
 {
-	// Do Nothing
+
 }
 
 /* Global functions
diff --git a/Engine/Source/Runtime/SlateCore/Private/Widgets/SWidget.cpp b/Engine/Source/Runtime/SlateCore/Private/Widgets/SWidget.cpp
--- a/Engine/Source/Runtime/SlateCore/Private/Widgets/SWidget.cpp
+++ b/Engine/Source/Runtime/SlateCore/Private/Widgets/SWidget.cpp
@@ -6,6 +6,7 @@
 #include "ActiveTimerHandle.h"
 #include "SlateStats.h"
 
+DECLARE_DWORD_ACCUMULATOR_STAT(TEXT("Total Widgets"), STAT_SlateTotalWidgets, STATGROUP_Slate);
 DECLARE_DWORD_COUNTER_STAT(TEXT("Num Painted Widgets"), STAT_SlateNumPaintedWidgets, STATGROUP_Slate);
 DECLARE_DWORD_COUNTER_STAT(TEXT("Num Ticked Widgets"), STAT_SlateNumTickedWidgets, STATGROUP_Slate);
 
@@ -37,7 +38,21 @@ SWidget::SWidget()
 	, bCachedVolatile(false)
 	, bInheritedVolatility(false)
 {
+	INC_DWORD_STAT(STAT_SlateTotalWidgets);
+}
 
+SWidget::~SWidget()
+{
+	// Unregister all ActiveTimers so they aren't left stranded in the Application's list.
+	if ( FSlateApplicationBase::IsInitialized() )
+	{
+		for ( const auto& ActiveTimerHandle : ActiveTimers )
+		{
+			FSlateApplicationBase::Get().UnRegisterActiveTimer(ActiveTimerHandle);
+		}
+	}
+
+	DEC_DWORD_STAT(STAT_SlateTotalWidgets);
 }
 
 void SWidget::Construct(
@@ -726,15 +741,3 @@ void SWidget::ExecuteActiveTimers(double CurrentTime, float DeltaTime)
 		}
 	}
 }
-
-SWidget::~SWidget()
-{
-	// Unregister all ActiveTimers so they aren't left stranded in the Application's list.
-	if (FSlateApplicationBase::IsInitialized())
-	{
-		for (const auto& ActiveTimerHandle : ActiveTimers)
-		{
-			FSlateApplicationBase::Get().UnRegisterActiveTimer(ActiveTimerHandle);
-		}
-	}
-}
diff --git a/Engine/Source/Runtime/SlateCore/Public/Rendering/DrawElements.h b/Engine/Source/Runtime/SlateCore/Public/Rendering/DrawElements.h
--- a/Engine/Source/Runtime/SlateCore/Public/Rendering/DrawElements.h
+++ b/Engine/Source/Runtime/SlateCore/Public/Rendering/DrawElements.h
@@ -75,6 +75,7 @@ public:
 
 	// Cached render data
 	class FSlateRenderDataHandle* CachedRenderData;
+	FVector2D CachedRenderDataOffset;
 
 	// Layer handle
 	FSlateDrawLayerHandle* LayerHandle;
@@ -166,9 +167,10 @@ public:
 		CustomDrawer = InCustomDrawer;
 	}
 
-	void SetCachedBufferPayloadProperties(FSlateRenderDataHandle* InRenderDataHandle)
+	void SetCachedBufferPayloadProperties(FSlateRenderDataHandle* InRenderDataHandle, const FVector2D& Offset)
 	{
 		CachedRenderData = InRenderDataHandle;
+		CachedRenderDataOffset = Offset;
 		checkSlow(CachedRenderData);
 	}
 
@@ -353,7 +355,7 @@ public:
 	 */
 	SLATECORE_API static void MakeCustom( FSlateWindowElementList& ElementList, uint32 InLayer, TSharedPtr<ICustomSlateElement, ESPMode::ThreadSafe> CustomDrawer );
 
-	SLATECORE_API static void MakeCachedBuffer(FSlateWindowElementList& ElementList, uint32 InLayer, TSharedPtr<FSlateRenderDataHandle, ESPMode::ThreadSafe>& CachedRenderDataHandle);
+	SLATECORE_API static void MakeCachedBuffer(FSlateWindowElementList& ElementList, uint32 InLayer, TSharedPtr<FSlateRenderDataHandle, ESPMode::ThreadSafe>& CachedRenderDataHandle, const FVector2D& Offset);
 
 	SLATECORE_API static void MakeLayer(FSlateWindowElementList& ElementList, uint32 InLayer, TSharedPtr<FSlateDrawLayerHandle, ESPMode::ThreadSafe>& DrawLayerHandle);
 
@@ -424,18 +426,29 @@ class FSlateRenderBatch;
 class SLATECORE_API FSlateRenderDataHandle : public TSharedFromThis < FSlateRenderDataHandle, ESPMode::ThreadSafe >
 {
 public:
-	FSlateRenderDataHandle(FSlateRenderer* InRenderer);
+	FSlateRenderDataHandle(const ILayoutCache* Cacher, FSlateRenderer* InRenderer);
 
 	virtual ~FSlateRenderDataHandle();
 
 	void Disconnect();
 
+	const ILayoutCache* GetCacher() const { return Cacher; }
+	const FSlateRenderer* GetRenderer() const { return Renderer; }
+
 	void SetRenderBatches(TArray<FSlateRenderBatch>* InRenderBatches) { RenderBatches = InRenderBatches; }
 	TArray<FSlateRenderBatch>* GetRenderBatches() { return RenderBatches; }
 
+	void BeginUsing() { FPlatformAtomics::InterlockedIncrement(&UsageCount); }
+	void EndUsing() { FPlatformAtomics::InterlockedDecrement(&UsageCount); }
+
+	bool IsInUse() const { return UsageCount > 0; }
+
 private:
+	const ILayoutCache* Cacher;
 	FSlateRenderer* Renderer;
 	TArray<FSlateRenderBatch>* RenderBatches;
+
+	volatile int32 UsageCount;
 };
 
 /** 
@@ -460,8 +473,8 @@ public:
 		, IndexArrayIndex(INDEX_NONE)
 	{}
 
-	FSlateElementBatch( TSharedPtr<FSlateRenderDataHandle, ESPMode::ThreadSafe> InCachedRenderHandle, const TOptional<FShortRect>& ScissorRect)
-		: BatchKey( InCachedRenderHandle, ScissorRect )
+	FSlateElementBatch( TSharedPtr<FSlateRenderDataHandle, ESPMode::ThreadSafe> InCachedRenderHandle, FVector2D InCachedRenderDataOffset, const TOptional<FShortRect>& ScissorRect)
+		: BatchKey( InCachedRenderHandle, InCachedRenderDataOffset, ScissorRect )
 		, ShaderResource( nullptr )
 		, NumElementsInBatch(0)
 		, VertexArrayIndex(INDEX_NONE)
@@ -495,12 +508,14 @@ public:
 	const TOptional<FShortRect>& GetScissorRect() const { return BatchKey.ScissorRect; }
 	const TWeakPtr<ICustomSlateElement, ESPMode::ThreadSafe> GetCustomDrawer() const { return BatchKey.CustomDrawer; }
 	const TSharedPtr<FSlateRenderDataHandle, ESPMode::ThreadSafe> GetCachedRenderHandle() const { return BatchKey.CachedRenderHandle; }
+	FVector2D GetCachedRenderDataOffset() const { return BatchKey.CachedRenderDataOffset; }
 	const TSharedPtr<FSlateDrawLayerHandle, ESPMode::ThreadSafe> GetLayerHandle() const { return BatchKey.LayerHandle; }
 private:
 	struct FBatchKey
 	{
 		const TWeakPtr<ICustomSlateElement, ESPMode::ThreadSafe> CustomDrawer;
 		const TSharedPtr<FSlateRenderDataHandle, ESPMode::ThreadSafe> CachedRenderHandle;
+		const FVector2D CachedRenderDataOffset;
 		const TSharedPtr<FSlateDrawLayerHandle, ESPMode::ThreadSafe> LayerHandle;
 		const FShaderParams ShaderParams;
 		const ESlateBatchDrawFlag::Type DrawFlags;	
@@ -529,8 +544,9 @@ private:
 			, ScissorRect( InScissorRect )
 		{}
 
-		FBatchKey( TSharedPtr<FSlateRenderDataHandle, ESPMode::ThreadSafe> InCachedRenderHandle, const TOptional<FShortRect>& InScissorRect)
+		FBatchKey( TSharedPtr<FSlateRenderDataHandle, ESPMode::ThreadSafe> InCachedRenderHandle, FVector2D InCachedRenderDataOffset, const TOptional<FShortRect>& InScissorRect)
 			: CachedRenderHandle(InCachedRenderHandle)
+			, CachedRenderDataOffset(InCachedRenderDataOffset)
 			, ShaderParams()
 			, DrawFlags(ESlateBatchDrawFlag::None)
 			, ShaderType(ESlateShader::Default)
@@ -615,15 +631,18 @@ public:
 		, NumIndices( InNumIndices )
 	{}
 
-	// Render Batch Sort Operator
-	bool operator < ( const FSlateRenderBatch& RenderBatchIn ) const
-	{
-		return Layer < RenderBatchIn.Layer;
-	}
+	//// Render Batch Sort Operator
+	//bool operator < ( const FSlateRenderBatch& RenderBatchIn ) const
+	//{
+	//	return Layer < RenderBatchIn.Layer;
+	//}
 
 public:
 	/** The layer we need to sort by when  */
-	uint32 Layer;
+	const uint32 Layer;
+
+	/** Dynamically modified offset that occurs when we have relative position stored render batches. */
+	FVector2D DynamicOffset;
 
 	const FShaderParams ShaderParams;
 
@@ -669,7 +688,8 @@ class FSlateBatchData
 {
 public:
 	FSlateBatchData()
-		: NumBatchedVertices(0)
+		: DynamicOffset(0, 0)
+		, NumBatchedVertices(0)
 		, NumBatchedIndices(0)
 		, NumLayers(0)
 	{}
@@ -722,26 +742,11 @@ public:
 	 */
 	SLATECORE_API void CreateRenderBatches(FElementBatchMap& LayerToElementBatches);
 
-	/**
-	 * Patches the render data to be offset by the new layer delta.
-	 */
-	SLATECORE_API void UpdateRenderBatches(FVector2D PositionOffset);
-
-	/** 
-	 * Sort the rendered element batches.
-	 */
-	SLATECORE_API void SortRenderBatches();
-
 private:
 	void Merge(FElementBatchMap& InLayerToElementBatches, uint32& VertexOffset, uint32& IndexOffset);
 
-	void AddRenderBatch( uint32 InLayer, const FSlateElementBatch& InElementBatch, int32 InNumVertices, int32 InNumIndices, int32 InVertexOffset, int32 InIndexOffset )
-	{
-		NumBatchedVertices += InNumVertices;
-		NumBatchedIndices += InNumIndices;
+	void AddRenderBatch(uint32 InLayer, const FSlateElementBatch& InElementBatch, int32 InNumVertices, int32 InNumIndices, int32 InVertexOffset, int32 InIndexOffset);
 
-		RenderBatches.Add( FSlateRenderBatch(InLayer, InElementBatch, RenderDataHandle, InNumVertices, InNumIndices, InVertexOffset, InIndexOffset) );
-	}
 private:
 
 	// The associated render data handle if these render batches are not in the default vertex/index buffer
@@ -762,6 +767,8 @@ private:
 	/** List of element batches sorted by later for use in rendering (for threaded renderers, can only be accessed from the render thread)*/
 	TArray<FSlateRenderBatch> RenderBatches;
 
+	FVector2D DynamicOffset;
+
 	int32 NumBatchedVertices;
 
 	int32 NumBatchedIndices;
@@ -972,18 +979,23 @@ public:
 	/**
 	 * Caches this element list on the renderer, generating all needed index and vertex buffers.
 	 */
-	SLATECORE_API TSharedRef<FSlateRenderDataHandle, ESPMode::ThreadSafe> CacheRenderData();
-
-	SLATECORE_API void UpdateCacheRenderData(FVector2D PositionOffset);
+	SLATECORE_API TSharedRef<FSlateRenderDataHandle, ESPMode::ThreadSafe> CacheRenderData(const ILayoutCache* Cacher);
 
 	SLATECORE_API TSharedPtr<FSlateRenderDataHandle, ESPMode::ThreadSafe> GetCachedRenderDataHandle() const
 	{
 		return CachedRenderDataHandle.Pin();
 	}
 
+	void BeginUsingCachedBuffer(TSharedPtr<FSlateRenderDataHandle, ESPMode::ThreadSafe>& InCachedRenderDataHandle)
+	{
+		InCachedRenderDataHandle->BeginUsing();
+		CachedRenderHandlesInUse.Add(InCachedRenderDataHandle);
+	}
+
 	SLATECORE_API bool IsCachedRenderDataInUse() const
 	{
-		return CachedRenderDataHandle.IsValid();
+		TSharedPtr<FSlateRenderDataHandle, ESPMode::ThreadSafe> SafeHandle = CachedRenderDataHandle.Pin();
+		return SafeHandle.IsValid() && SafeHandle->IsInUse();
 	}
 
 	SLATECORE_API void PreDraw_RenderThread();
@@ -1015,6 +1027,8 @@ private:
 	TArray<FSlateDrawElement*> DrawElementFreePool;
 #endif
 
+	TArray< TSharedPtr<FSlateRenderDataHandle, ESPMode::ThreadSafe> > CachedRenderHandlesInUse;
+
 	/**
 	 * Some widgets want their logical children to appear at a different "layer" in the physical hierarchy.
 	 * We accomplish this by deferring their painting.
diff --git a/Engine/Source/Runtime/SlateCore/Public/Rendering/SlateRenderer.h b/Engine/Source/Runtime/SlateCore/Public/Rendering/SlateRenderer.h
--- a/Engine/Source/Runtime/SlateCore/Public/Rendering/SlateRenderer.h
+++ b/Engine/Source/Runtime/SlateCore/Public/Rendering/SlateRenderer.h
@@ -237,12 +237,14 @@ public:
 	 * Converts and caches the elements list data as the final rendered vertex and index buffer data,
 	 * returning a handle to it to allow issuing future draw commands using it.
 	 */
-	virtual TSharedRef<FSlateRenderDataHandle, ESPMode::ThreadSafe> CacheElementRenderData(FSlateWindowElementList& ElementList);
+	virtual TSharedRef<FSlateRenderDataHandle, ESPMode::ThreadSafe> CacheElementRenderData(const ILayoutCache* Cacher, FSlateWindowElementList& ElementList);
 
 	/**
-	 * Patch existing render data to have a new layer delta applied to all elements.
+	 * Releases the caching resources used on the render thread for the provided ILayoutCache.  This is
+	 * should be the kind of operation you perform when the ILayoutCache is being destroyed and will no longer
+	 * be needing to draw anything again.
 	 */
-	virtual void UpdateElementRenderData(FSlateWindowElementList& ElementList, FVector2D PositionOffset);
+	virtual void ReleaseCachingResourcesFor(const ILayoutCache* Cacher);
 
 	/**
 	 * You must call this before calling CopyWindowsToVirtualScreenBuffer(), to setup the render targets first.
diff --git a/Engine/Source/Runtime/SlateCore/Public/Stats/SlateStats.h b/Engine/Source/Runtime/SlateCore/Public/Stats/SlateStats.h
--- a/Engine/Source/Runtime/SlateCore/Public/Stats/SlateStats.h
+++ b/Engine/Source/Runtime/SlateCore/Public/Stats/SlateStats.h
@@ -3,7 +3,8 @@
 #pragma once
 
 /** Set to zero to disable the slate stats system completely, meaning all stats-related structures and code are compiled out. */
-#define SLATE_STATS (!UE_BUILD_SHIPPING || WITH_EDITOR)
+//#define SLATE_STATS (!(UE_BUILD_SHIPPING || UE_BUILD_TEST) || WITH_EDITOR)
+#define SLATE_STATS (WITH_SLATE_STATS || WITH_EDITOR)
 
 /** 
  * Predefined detail levels to make choosing level slightly easier.
diff --git a/Engine/Source/Runtime/SlateCore/Public/Types/ReflectionMetadata.h b/Engine/Source/Runtime/SlateCore/Public/Types/ReflectionMetadata.h
--- a/Engine/Source/Runtime/SlateCore/Public/Types/ReflectionMetadata.h
+++ b/Engine/Source/Runtime/SlateCore/Public/Types/ReflectionMetadata.h
@@ -14,7 +14,7 @@ class FReflectionMetaData : public ISlateMetaData
 public:
 	SLATE_METADATA_TYPE(FReflectionMetaData, ISlateMetaData)
 
-	FReflectionMetaData(FName InName, UClass* InClass, UObject* InAsset)
+	FReflectionMetaData(FName InName, UClass* InClass, const UObject* InAsset)
 		: Name(InName)
 		, Class(InClass)
 		, Asset(InAsset)
@@ -28,11 +28,11 @@ public:
 	TWeakObjectPtr<UClass> Class;
 
 	/** The asset that owns the widget and is responsible for its specific existence. */
-	TWeakObjectPtr<UObject> Asset;
+	TWeakObjectPtr<const UObject> Asset;
 
 public:
 
-	static FString GetWidgetDebugInfo(TSharedPtr<const SWidget>& InWidget)
+	static FString GetWidgetDebugInfo(const SWidget* InWidget)
 	{
 		// UMG widgets have meta-data to help track them
 		TSharedPtr<FReflectionMetaData> MetaData = InWidget->GetMetaData<FReflectionMetaData>();
diff --git a/Engine/Source/Runtime/SlateCore/Public/Widgets/SWidget.h b/Engine/Source/Runtime/SlateCore/Public/Widgets/SWidget.h
--- a/Engine/Source/Runtime/SlateCore/Public/Widgets/SWidget.h
+++ b/Engine/Source/Runtime/SlateCore/Public/Widgets/SWidget.h
@@ -68,17 +68,26 @@ public:
 
 
 /**
- * 
+ * The different types of invalidation that are possible for a widget.
  */
 enum class EInvalidateWidget
 {
+	/**
+	 * Use Layout invalidation if you're changing a normal property involving painting or sizing.
+	 */
 	Layout,
+	/**
+	 * Use Layout invalidation if you're changing a normal property involving painting or sizing.
+	 * Additionally if the property that was changed affects Volatility in anyway, it's important
+	 * that you invalidate volatility so that it can be recalculated and cached.
+	 */
 	LayoutAndVolatility
 };
 
 
 /**
- * 
+ * An ILayoutCache implementor is responsible for caching a the hierarchy of widgets it is drawing.
+ * The shipped implementation of this is SInvalidationPanel.
  */
 class ILayoutCache
 {
diff --git a/Engine/Source/Runtime/SlateRHIRenderer/Private/SlateRHIRenderer.cpp b/Engine/Source/Runtime/SlateRHIRenderer/Private/SlateRHIRenderer.cpp
--- a/Engine/Source/Runtime/SlateRHIRenderer/Private/SlateRHIRenderer.cpp
+++ b/Engine/Source/Runtime/SlateRHIRenderer/Private/SlateRHIRenderer.cpp
@@ -11,8 +11,10 @@ DECLARE_CYCLE_STAT(TEXT("Map Staging Buffer"),STAT_MapStagingBuffer,STATGROUP_Cr
 DECLARE_CYCLE_STAT(TEXT("Generate Capture Buffer"),STAT_GenerateCaptureBuffer,STATGROUP_CrashTracker);
 DECLARE_CYCLE_STAT(TEXT("Unmap Staging Buffer"),STAT_UnmapStagingBuffer,STATGROUP_CrashTracker);
 
-DECLARE_CYCLE_STAT(TEXT("Slate Rendering RT Time"), STAT_SlateRenderingRTTime, STATGROUP_Slate);
-DECLARE_CYCLE_STAT(TEXT("Slate RT Present Time"), STAT_SlatePresentRTTime, STATGROUP_Slate);
+DECLARE_CYCLE_STAT(TEXT("Slate RT: Rendering"), STAT_SlateRenderingRTTime, STATGROUP_Slate);
+DECLARE_CYCLE_STAT(TEXT("Slate RT: Create Batches"), STAT_SlateRTCreateBatches, STATGROUP_Slate);
+DECLARE_CYCLE_STAT(TEXT("Slate RT: Fill Vertex & Index Buffers"), STAT_SlateRTFillVertexIndexBuffers, STATGROUP_Slate);
+DECLARE_CYCLE_STAT(TEXT("Slate RT: Draw Batches"), STAT_SlateRTDrawBatches, STATGROUP_Slate);
 
 // Defines the maximum size that a slate viewport will create
 #define MAX_VIEWPORT_SIZE 16384
@@ -407,9 +409,17 @@ void FSlateRHIRenderer::DrawWindow_RenderThread(FRHICommandListImmediate& RHICmd
 
 		WindowElementList.PreDraw_RenderThread();
 
-		// Update the vertex and index buffer	
-		BatchData.CreateRenderBatches(RootBatchMap);
-		RenderingPolicy->UpdateVertexAndIndexBuffers(RHICmdList, BatchData);
+		{
+			SCOPE_CYCLE_COUNTER(STAT_SlateRTCreateBatches);
+			// Update the vertex and index buffer	
+			BatchData.CreateRenderBatches(RootBatchMap);
+		}
+
+		{
+			SCOPE_CYCLE_COUNTER(STAT_SlateRTFillVertexIndexBuffers);
+			RenderingPolicy->UpdateVertexAndIndexBuffers(RHICmdList, BatchData);
+		}
+
 		// should have been created by the game thread
 		check( IsValidRef(ViewportInfo.ViewportRHI) );
 
@@ -447,6 +457,8 @@ void FSlateRHIRenderer::DrawWindow_RenderThread(FRHICommandListImmediate& RHICmd
 #endif
 		if( BatchData.GetRenderBatches().Num() > 0 )
 		{
+			SCOPE_CYCLE_COUNTER(STAT_SlateRTDrawBatches);
+
 			FSlateBackBuffer BackBufferTarget( BackBuffer, FIntPoint( ViewportWidth, ViewportHeight ) );
 
 			RenderingPolicy->DrawElements
@@ -1241,9 +1253,9 @@ void FSlateRHIRenderer::SetWindowRenderTarget(const SWindow& Window, IViewportRe
 	}
 }
 
-TSharedRef<FSlateRenderDataHandle, ESPMode::ThreadSafe> FSlateRHIRenderer::CacheElementRenderData(FSlateWindowElementList& ElementList)
+TSharedRef<FSlateRenderDataHandle, ESPMode::ThreadSafe> FSlateRHIRenderer::CacheElementRenderData(const ILayoutCache* Cacher, FSlateWindowElementList& ElementList)
 {
-	TSharedRef<FSlateRenderDataHandle, ESPMode::ThreadSafe> RenderDataHandle = MakeShareable(new FSlateRenderDataHandle(this));
+	TSharedRef<FSlateRenderDataHandle, ESPMode::ThreadSafe> RenderDataHandle = MakeShareable(new FSlateRenderDataHandle(Cacher, this));
 
 	checkSlow(ElementList.GetChildDrawLayers().Num() == 0);
 
@@ -1280,35 +1292,24 @@ TSharedRef<FSlateRenderDataHandle, ESPMode::ThreadSafe> FSlateRHIRenderer::Cache
 	return RenderDataHandle;
 }
 
-void FSlateRHIRenderer::UpdateElementRenderData(FSlateWindowElementList& ElementList, FVector2D PositionOffset)
+void FSlateRHIRenderer::ReleaseCachingResourcesFor(const ILayoutCache* Cacher)
 {
-	struct FUpdateCacheElementBatchesContext
+	struct FReleaseCachingResourcesForContext
 	{
 		FSlateRHIRenderingPolicy* RenderPolicy;
-		FSlateWindowElementList* SlateElementList;
-		TSharedPtr<FSlateRenderDataHandle, ESPMode::ThreadSafe> RenderDataHandle;
-		FVector2D PositionOffset;
+		const ILayoutCache* Cacher;
 	};
-	FUpdateCacheElementBatchesContext UpdateCacheElementBatchesContext =
+	FReleaseCachingResourcesForContext MarshalContext =
 	{
 		RenderingPolicy.Get(),
-		&ElementList,
-		ElementList.GetCachedRenderDataHandle(),
-		PositionOffset
+		Cacher,
 	};
 	ENQUEUE_UNIQUE_RENDER_COMMAND_ONEPARAMETER(
-		UpdateCacheElementBatches,
-		FUpdateCacheElementBatchesContext, Context, UpdateCacheElementBatchesContext,
-	{
-		if ( Context.RenderDataHandle.IsValid() )
+		ReleaseCachingResourcesFor,
+		FReleaseCachingResourcesForContext, Context, MarshalContext,
 		{
-			auto RenderDataHandleRef = Context.RenderDataHandle.ToSharedRef();
-
-			FSlateBatchData& BatchData = Context.SlateElementList->GetBatchData();
-			BatchData.UpdateRenderBatches(Context.PositionOffset);
-			Context.RenderPolicy->UpdateVertexAndIndexBuffers(RHICmdList, BatchData, RenderDataHandleRef);
-		}
-	});
+			Context.RenderPolicy->ReleaseCachingResourcesFor(Context.Cacher);
+		});
 }
 
 void FSlateRHIRenderer::ReleaseCachedRenderData(FSlateRenderDataHandle* InRenderHandle)
diff --git a/Engine/Source/Runtime/SlateRHIRenderer/Private/SlateRHIRenderer.h b/Engine/Source/Runtime/SlateRHIRenderer/Private/SlateRHIRenderer.h
--- a/Engine/Source/Runtime/SlateRHIRenderer/Private/SlateRHIRenderer.h
+++ b/Engine/Source/Runtime/SlateRHIRenderer/Private/SlateRHIRenderer.h
@@ -188,8 +188,8 @@ public:
 	virtual void ReleaseUpdatableTexture(FSlateUpdatableTexture* Texture) override;
 	virtual ISlateAtlasProvider* GetTextureAtlasProvider() override;
 	virtual void ReleaseAccessedResources() override;
-	virtual TSharedRef<FSlateRenderDataHandle, ESPMode::ThreadSafe> CacheElementRenderData(FSlateWindowElementList& ElementList) override;
-	virtual void UpdateElementRenderData(FSlateWindowElementList& ElementList, FVector2D PositionOffset) override;
+	virtual TSharedRef<FSlateRenderDataHandle, ESPMode::ThreadSafe> CacheElementRenderData(const ILayoutCache* Cacher, FSlateWindowElementList& ElementList) override;
+	virtual void ReleaseCachingResourcesFor(const ILayoutCache* Cacher) override;
 	virtual void ReleaseCachedRenderData(FSlateRenderDataHandle* RenderHandle) override;
 
 	/** Draws windows from a FSlateDrawBuffer on the render thread */
diff --git a/Engine/Source/Runtime/SlateRHIRenderer/Private/SlateRHIRenderingPolicy.cpp b/Engine/Source/Runtime/SlateRHIRenderer/Private/SlateRHIRenderingPolicy.cpp
--- a/Engine/Source/Runtime/SlateRHIRenderer/Private/SlateRHIRenderingPolicy.cpp
+++ b/Engine/Source/Runtime/SlateRHIRenderer/Private/SlateRHIRenderingPolicy.cpp
@@ -14,13 +14,10 @@
 
 DECLARE_CYCLE_STAT(TEXT("Update Buffers RT"), STAT_SlateUpdateBufferRTTime, STATGROUP_Slate);
 DECLARE_CYCLE_STAT(TEXT("PreFill Buffers RT"), STAT_SlatePreFullBufferRTTime, STATGROUP_Slate);
-DECLARE_CYCLE_STAT(TEXT("Draw Time RT"), STAT_SlateDrawTime, STATGROUP_Slate);
 DECLARE_DWORD_COUNTER_STAT(TEXT("Num Layers"), STAT_SlateNumLayers, STATGROUP_Slate);
 DECLARE_DWORD_COUNTER_STAT(TEXT("Num Batches"), STAT_SlateNumBatches, STATGROUP_Slate);
 DECLARE_DWORD_COUNTER_STAT(TEXT("Num Vertices"), STAT_SlateVertexCount, STATGROUP_Slate);
 
-DECLARE_MEMORY_STAT(TEXT("Batch Vertex Memory"), STAT_SlateVertexBatchMemory, STATGROUP_SlateMemory);
-DECLARE_MEMORY_STAT(TEXT("Batch Index Memory"), STAT_SlateIndexBatchMemory, STATGROUP_SlateMemory);
 DECLARE_MEMORY_STAT(TEXT("Vertex Buffer Memory"), STAT_SlateVertexBufferMemory, STATGROUP_SlateMemory);
 DECLARE_MEMORY_STAT(TEXT("Index Buffer Memory"), STAT_SlateIndexBufferMemory, STATGROUP_SlateMemory);
 
@@ -202,13 +199,23 @@ void FSlateRHIRenderingPolicy::ReleaseResources()
 
 		Buffer->VertexBuffer.Destroy();
 		Buffer->IndexBuffer.Destroy();
+		delete Buffer;
 	}
 
-	for ( FCachedRenderBuffers* PooledBuffer : CachedBufferPool )
+	CachedBuffers.Reset();
+
+	for ( TCachedBufferPoolMap::TIterator BufferIt(CachedBufferPool); BufferIt; ++BufferIt )
 	{
-		PooledBuffer->VertexBuffer.Destroy();
-		PooledBuffer->IndexBuffer.Destroy();
+		TArray< FCachedRenderBuffers* >& Pool = BufferIt.Value();
+		for ( FCachedRenderBuffers* PooledBuffer : Pool )
+		{
+			PooledBuffer->VertexBuffer.Destroy();
+			PooledBuffer->IndexBuffer.Destroy();
+			delete PooledBuffer;
+		}
 	}
+
+	CachedBufferPool.Reset();
 }
 
 void FSlateRHIRenderingPolicy::BeginDrawingWindows()
@@ -220,6 +227,7 @@ void FSlateRHIRenderingPolicy::EndDrawingWindows()
 {
 	check( IsInRenderingThread() );
 
+#if STATS
 	uint32 TotalVertexBufferMemory = 0;
 	uint32 TotalIndexBufferMemory = 0;
 
@@ -238,9 +246,33 @@ void FSlateRHIRenderingPolicy::EndDrawingWindows()
 		TotalIndexBufferUsage += IndexBuffers[BufferIndex].GetBufferUsageSize();
 	}
 
+	for ( TCachedBufferMap::TIterator BufferIt(CachedBuffers); BufferIt; ++BufferIt )
+	{
+		FCachedRenderBuffers* PooledBuffer = BufferIt.Value();
+
+		TotalVertexBufferMemory += PooledBuffer->VertexBuffer.GetBufferSize();
+		TotalVertexBufferUsage += PooledBuffer->VertexBuffer.GetBufferUsageSize();
+
+		TotalIndexBufferMemory += PooledBuffer->IndexBuffer.GetBufferSize();
+		TotalIndexBufferUsage += PooledBuffer->IndexBuffer.GetBufferUsageSize();
+	}
+
+	for ( TCachedBufferPoolMap::TIterator BufferIt(CachedBufferPool); BufferIt; ++BufferIt )
+	{
+		TArray< FCachedRenderBuffers* >& Pool = BufferIt.Value();
+		for ( FCachedRenderBuffers* PooledBuffer : Pool )
+		{
+			TotalVertexBufferMemory += PooledBuffer->VertexBuffer.GetBufferSize();
+			TotalVertexBufferUsage += PooledBuffer->VertexBuffer.GetBufferUsageSize();
+
+			TotalIndexBufferMemory += PooledBuffer->IndexBuffer.GetBufferSize();
+			TotalIndexBufferUsage += PooledBuffer->IndexBuffer.GetBufferUsageSize();
+		}
+	}
 
 	SET_MEMORY_STAT( STAT_SlateVertexBufferMemory, TotalVertexBufferMemory );
 	SET_MEMORY_STAT( STAT_SlateIndexBufferMemory, TotalIndexBufferMemory );
+#endif
 }
 
 struct FSlateUpdateVertexAndIndexBuffers : public FRHICommand<FSlateUpdateVertexAndIndexBuffers>
@@ -266,8 +298,6 @@ struct FSlateUpdateVertexAndIndexBuffers : public FRHICommand<FSlateUpdateVertex
 	
 		VertexBuffer.UnlockBuffer_RHIThread();
 		IndexBuffer.UnlockBuffer_RHIThread();
-
-		BatchData.SortRenderBatches();
 	}
 };
 
@@ -279,13 +309,23 @@ void FSlateRHIRenderingPolicy::UpdateVertexAndIndexBuffers(FRHICommandListImmedi
 	UpdateVertexAndIndexBuffers(RHICmdList, InBatchData, VertexBuffer, IndexBuffer);
 }
 
-void FSlateRHIRenderingPolicy::UpdateVertexAndIndexBuffers(FRHICommandListImmediate& RHICmdList, FSlateBatchData& InBatchData, TSharedRef<FSlateRenderDataHandle, ESPMode::ThreadSafe> RenderHandle)
+void FSlateRHIRenderingPolicy::UpdateVertexAndIndexBuffers(FRHICommandListImmediate& RHICmdList, FSlateBatchData& InBatchData, const TSharedRef<FSlateRenderDataHandle, ESPMode::ThreadSafe>& RenderHandle)
 {
+	// Should only be called by the rendering thread
+	check(IsInRenderingThread());
+
 	FCachedRenderBuffers* Buffers = CachedBuffers.FindRef(&RenderHandle.Get());
 	if ( Buffers == nullptr )
 	{
+		// Rather than having a global pool, we associate the pools with a particular layout cacher.
+		// If we don't do this, all buffers eventually become as larger as the largest buffer, and it
+		// would be much better to keep the pools coherent with the sizes typically associated with
+		// a particular caching panel.
+		const ILayoutCache* LayoutCacher = RenderHandle->GetCacher();
+		TArray< FCachedRenderBuffers* >& Pool = CachedBufferPool.FindOrAdd(LayoutCacher);
+
 		// If the cached buffer pool is empty, time to create a new one!
-		if ( CachedBufferPool.Num() == 0 )
+		if ( Pool.Num() == 0 )
 		{
 			Buffers = new FCachedRenderBuffers();
 			Buffers->VertexBuffer.Init(200);
@@ -294,8 +334,8 @@ void FSlateRHIRenderingPolicy::UpdateVertexAndIndexBuffers(FRHICommandListImmedi
 		else
 		{
 			// If we found one in the pool, lets use it!
-			Buffers = CachedBufferPool[0];
-			CachedBufferPool.RemoveAtSwap(0, 1, false);
+			Buffers = Pool[0];
+			Pool.RemoveAtSwap(0, 1, false);
 		}
 
 		CachedBuffers.Add(&RenderHandle.Get(), Buffers);
@@ -306,14 +346,47 @@ void FSlateRHIRenderingPolicy::UpdateVertexAndIndexBuffers(FRHICommandListImmedi
 
 void FSlateRHIRenderingPolicy::ReleaseCachedRenderData(FSlateRenderDataHandle* InRenderHandle)
 {
-	FCachedRenderBuffers* Buffers = CachedBuffers.FindRef(InRenderHandle);
-	if ( ensure(Buffers != nullptr) )
+	// Should only be called by the rendering thread
+	check(IsInRenderingThread());
+	check(InRenderHandle);
+
+	FCachedRenderBuffers* PooledBuffer = CachedBuffers.FindRef(InRenderHandle);
+	if ( ensure(PooledBuffer != nullptr) )
 	{
-		CachedBufferPool.Add(Buffers);
+		const ILayoutCache* LayoutCacher = InRenderHandle->GetCacher();
+		TArray< FCachedRenderBuffers* >* Pool = CachedBufferPool.Find(LayoutCacher);
+		if ( Pool )
+		{
+			Pool->Add(PooledBuffer);
+		}
+		else
+		{
+			// The buffer pool may have already been released, so lets just delete this buffer.
+			PooledBuffer->VertexBuffer.Destroy();
+			PooledBuffer->IndexBuffer.Destroy();
+			delete PooledBuffer;
+		}
+		
 		CachedBuffers.Remove(InRenderHandle);
 	}
 }
 
+void FSlateRHIRenderingPolicy::ReleaseCachingResourcesFor(const ILayoutCache* Cacher)
+{
+	TArray< FCachedRenderBuffers* >* Pool = CachedBufferPool.Find(Cacher);
+	if ( Pool )
+	{
+		for ( FCachedRenderBuffers* PooledBuffer : *Pool )
+		{
+			PooledBuffer->VertexBuffer.Destroy();
+			PooledBuffer->IndexBuffer.Destroy();
+			delete PooledBuffer;
+		}
+
+		CachedBufferPool.Remove(Cacher);
+	}
+}
+
 void FSlateRHIRenderingPolicy::UpdateVertexAndIndexBuffers(FRHICommandListImmediate& RHICmdList, FSlateBatchData& InBatchData, TSlateElementVertexBuffer<FSlateVertex>& VertexBuffer, FSlateElementIndexBuffer& IndexBuffer)
 {
 	SCOPE_CYCLE_COUNTER( STAT_SlateUpdateBufferRTTime );
@@ -331,7 +404,7 @@ void FSlateRHIRenderingPolicy::UpdateVertexAndIndexBuffers(FRHICommandListImmedi
 		VertexBuffer.PreFillBuffer(NumVertices, bShouldShrinkResources);
 		IndexBuffer.PreFillBuffer(NumIndices, bShouldShrinkResources);
 
-		if(1 || !GRHIThread || RHICmdList.Bypass())
+		if(!GRHIThread || RHICmdList.Bypass())
 		{
 			uint8* VertexBufferData = (uint8*)VertexBuffer.LockBuffer_RenderThread(NumVertices);
 			uint8* IndexBufferData =  (uint8*)IndexBuffer.LockBuffer_RenderThread(NumIndices);
@@ -340,8 +413,6 @@ void FSlateRHIRenderingPolicy::UpdateVertexAndIndexBuffers(FRHICommandListImmedi
 	
 			VertexBuffer.UnlockBuffer_RenderThread();
 			IndexBuffer.UnlockBuffer_RenderThread();
-
-			InBatchData.SortRenderBatches();
 		}
 		else
 		{
@@ -432,8 +503,6 @@ static FSceneView& CreateSceneView( FSceneViewFamilyContext& ViewFamilyContext,
 
 void FSlateRHIRenderingPolicy::DrawElements(FRHICommandListImmediate& RHICmdList, FSlateBackBuffer& BackBuffer, const FMatrix& ViewProjectionMatrix, const TArray<FSlateRenderBatch>& RenderBatches, bool bAllowSwitchVerticalAxis)
 {
-	SCOPE_CYCLE_COUNTER( STAT_SlateDrawTime );
-
 	// Should only be called by the rendering thread
 	check(IsInRenderingThread());
 
@@ -542,8 +611,9 @@ void FSlateRHIRenderingPolicy::DrawElements(FRHICommandListImmediate& RHICmdList
 					PixelShader->GetPixelShader(),
 					FGeometryShaderRHIRef()));
 
+				FMatrix DynamicOffset = FTranslationMatrix::Make(FVector(RenderBatch.DynamicOffset.X, RenderBatch.DynamicOffset.Y, 0));
 
-				VertexShader->SetViewProjection(RHICmdList, ViewProjectionMatrix);
+				VertexShader->SetViewProjection(RHICmdList, DynamicOffset * ViewProjectionMatrix);
 				VertexShader->SetVerticalAxisMultiplier(RHICmdList, bAllowSwitchVerticalAxis && RHINeedsToSwitchVerticalAxis(GShaderPlatformForFeatureLevel[GMaxRHIFeatureLevel]) ? -1.0f : 1.0f );
 #if !DEBUG_OVERDRAW
 				RHICmdList.SetBlendState(
diff --git a/Engine/Source/Runtime/SlateRHIRenderer/Private/SlateRHIRenderingPolicy.h b/Engine/Source/Runtime/SlateRHIRenderer/Private/SlateRHIRenderingPolicy.h
--- a/Engine/Source/Runtime/SlateRHIRenderer/Private/SlateRHIRenderingPolicy.h
+++ b/Engine/Source/Runtime/SlateRHIRenderer/Private/SlateRHIRenderingPolicy.h
@@ -246,9 +246,10 @@ public:
 	~FSlateRHIRenderingPolicy();
 
 	void UpdateVertexAndIndexBuffers(FRHICommandListImmediate& RHICmdList, FSlateBatchData& BatchData);
-	void UpdateVertexAndIndexBuffers(FRHICommandListImmediate& RHICmdList, FSlateBatchData& BatchData, TSharedRef<FSlateRenderDataHandle, ESPMode::ThreadSafe> RenderHandle);
+	void UpdateVertexAndIndexBuffers(FRHICommandListImmediate& RHICmdList, FSlateBatchData& BatchData, const TSharedRef<FSlateRenderDataHandle, ESPMode::ThreadSafe>& RenderHandle);
 
 	void ReleaseCachedRenderData(FSlateRenderDataHandle* InRenderHandle);
+	void ReleaseCachingResourcesFor(const ILayoutCache* Cacher);
 
 	virtual void DrawElements(FRHICommandListImmediate& RHICmdList, class FSlateBackBuffer& BackBuffer, const FMatrix& ViewProjectionMatrix, const TArray<FSlateRenderBatch>& RenderBatches, bool bAllowSwtichVerticalAxis=true);
 
@@ -292,7 +293,9 @@ private:
 
 	typedef TMap< FSlateRenderDataHandle*, FCachedRenderBuffers* > TCachedBufferMap;
 	TCachedBufferMap CachedBuffers;
-	TArray< FCachedRenderBuffers* > CachedBufferPool;
+
+	typedef TMap< const ILayoutCache*, TArray< FCachedRenderBuffers* > > TCachedBufferPoolMap;
+	TCachedBufferPoolMap CachedBufferPool;
 
 	TSharedRef<FSlateRHIResourceManager> ResourceManager;
 	TSharedPtr<FSlateFontCache> FontCache;
diff --git a/Engine/Source/Runtime/UMG/Private/Components/Widget.cpp b/Engine/Source/Runtime/UMG/Private/Components/Widget.cpp
--- a/Engine/Source/Runtime/UMG/Private/Components/Widget.cpp
+++ b/Engine/Source/Runtime/UMG/Private/Components/Widget.cpp
@@ -800,6 +800,12 @@ void UWidget::SynchronizeProperties()
 	// In editor builds we add metadata to the widget so that once hit with the widget reflector it can report
 	// where it comes from, what blueprint, what the name of the widget was...etc.
 	SafeWidget->AddMetadata<FReflectionMetaData>(MakeShareable(new FReflectionMetaData(GetFName(), GetClass(), WidgetGeneratedBy)));
+#else
+
+#if !UE_BUILD_SHIPPING
+	SafeWidget->AddMetadata<FReflectionMetaData>(MakeShareable(new FReflectionMetaData(GetFName(), GetClass(), WidgetGeneratedByClass)));
+#endif
+
 #endif
 }
 
diff --git a/Engine/Source/Runtime/UMG/Private/WidgetBlueprintGeneratedClass.cpp b/Engine/Source/Runtime/UMG/Private/WidgetBlueprintGeneratedClass.cpp
--- a/Engine/Source/Runtime/UMG/Private/WidgetBlueprintGeneratedClass.cpp
+++ b/Engine/Source/Runtime/UMG/Private/WidgetBlueprintGeneratedClass.cpp
@@ -26,6 +26,8 @@ void UWidgetBlueprintGeneratedClass::InitializeWidget(UUserWidget* UserWidget) c
 	UWidgetTree* ClonedTree = DuplicateObject<UWidgetTree>(WidgetTree, UserWidget);
 #endif
 
+	UserWidget->WidgetGeneratedByClass = this;
+
 #if WITH_EDITOR
 	UserWidget->WidgetGeneratedBy = ClassGeneratedBy;
 #endif
@@ -60,6 +62,8 @@ void UWidgetBlueprintGeneratedClass::InitializeWidget(UUserWidget* UserWidget) c
 				return;
 			}
 
+			Widget->WidgetGeneratedByClass = this;
+
 #if WITH_EDITOR
 			Widget->WidgetGeneratedBy = ClassGeneratedBy;
 #endif
diff --git a/Engine/Source/Runtime/UMG/Public/Components/Widget.h b/Engine/Source/Runtime/UMG/Public/Components/Widget.h
--- a/Engine/Source/Runtime/UMG/Public/Components/Widget.h
+++ b/Engine/Source/Runtime/UMG/Public/Components/Widget.h
@@ -225,10 +225,14 @@ public:
 
 	/** Stores a reference to the asset responsible for this widgets construction. */
 	UPROPERTY(Transient)
-	UObject* WidgetGeneratedBy;
+	const UObject* WidgetGeneratedBy;
 
 #endif
 
+	/** Stores a reference to the class responsible for this widgets construction. */
+	UPROPERTY(Transient)
+	const UClass* WidgetGeneratedByClass;
+
 public:
 
 	/** */
