Merged CL 2664420

Fix cached invalidation resources not being properly deleted at shutdown.  Causes a shutdown crash

diff --git a/Engine/Source/Runtime/Slate/Private/Widgets/SInvalidationPanel.cpp b/Engine/Source/Runtime/Slate/Private/Widgets/SInvalidationPanel.cpp
--- a/Engine/Source/Runtime/Slate/Private/Widgets/SInvalidationPanel.cpp
+++ b/Engine/Source/Runtime/Slate/Private/Widgets/SInvalidationPanel.cpp
@@ -64,8 +64,11 @@ SInvalidationPanel::~SInvalidationPanel()
 		delete NodePool[i];
 	}
 
-	TSharedPtr<FSlateRenderer> Renderer = FSlateApplicationBase::Get().GetRenderer();
-	Renderer->ReleaseCachingResourcesFor(this);
+	if( FSlateApplication::IsInitialized() )
+	{
+		TSharedPtr<FSlateRenderer> Renderer = FSlateApplicationBase::Get().GetRenderer();
+		Renderer->ReleaseCachingResourcesFor(this);
+	}
 }
 
 bool SInvalidationPanel::GetCanCache() const
diff --git a/Engine/Source/Runtime/SlateRHIRenderer/Private/SlateRHIRenderingPolicy.cpp b/Engine/Source/Runtime/SlateRHIRenderer/Private/SlateRHIRenderingPolicy.cpp
--- a/Engine/Source/Runtime/SlateRHIRenderer/Private/SlateRHIRenderingPolicy.cpp
+++ b/Engine/Source/Runtime/SlateRHIRenderer/Private/SlateRHIRenderingPolicy.cpp
@@ -155,6 +155,8 @@ FSlateRHIRenderingPolicy::FSlateRHIRenderingPolicy( TSharedPtr<FSlateFontCache>
 
 FSlateRHIRenderingPolicy::~FSlateRHIRenderingPolicy()
 {
+	// Delete released resources.  Note this MUST NOT be called before the rendering resources have been released
+	DeleteReleasedResources();
 }
 
 void FSlateRHIRenderingPolicy::InitResources()
@@ -199,11 +201,8 @@ void FSlateRHIRenderingPolicy::ReleaseResources()
 
 		Buffer->VertexBuffer.Destroy();
 		Buffer->IndexBuffer.Destroy();
-		delete Buffer;
 	}
 
-	CachedBuffers.Reset();
-
 	for ( TCachedBufferPoolMap::TIterator BufferIt(CachedBufferPool); BufferIt; ++BufferIt )
 	{
 		TArray< FCachedRenderBuffers* >& Pool = BufferIt.Value();
@@ -211,11 +210,30 @@ void FSlateRHIRenderingPolicy::ReleaseResources()
 		{
 			PooledBuffer->VertexBuffer.Destroy();
 			PooledBuffer->IndexBuffer.Destroy();
+		}
+	}
+}
+
+void FSlateRHIRenderingPolicy::DeleteReleasedResources()
+{
+	for ( TCachedBufferMap::TIterator BufferIt(CachedBuffers); BufferIt; ++BufferIt )
+	{
+		FCachedRenderBuffers* Buffer = BufferIt.Value();
+		delete Buffer;
+	}
+
+	CachedBuffers.Empty();
+
+	for ( TCachedBufferPoolMap::TIterator BufferIt(CachedBufferPool); BufferIt; ++BufferIt )
+	{
+		TArray< FCachedRenderBuffers* >& Pool = BufferIt.Value();
+		for ( FCachedRenderBuffers* PooledBuffer : Pool )
+		{
 			delete PooledBuffer;
 		}
 	}
 
-	CachedBufferPool.Reset();
+	CachedBufferPool.Empty();
 }
 
 void FSlateRHIRenderingPolicy::BeginDrawingWindows()
diff --git a/Engine/Source/Runtime/SlateRHIRenderer/Private/SlateRHIRenderingPolicy.h b/Engine/Source/Runtime/SlateRHIRenderer/Private/SlateRHIRenderingPolicy.h
--- a/Engine/Source/Runtime/SlateRHIRenderer/Private/SlateRHIRenderingPolicy.h
+++ b/Engine/Source/Runtime/SlateRHIRenderer/Private/SlateRHIRenderingPolicy.h
@@ -266,7 +266,6 @@ public:
 	void SetUseGammaCorrection( bool bInUseGammaCorrection ) { bGammaCorrect = bInUseGammaCorrection; }
 protected:
 	void UpdateVertexAndIndexBuffers(FRHICommandListImmediate& RHICmdList, FSlateBatchData& BatchData, TSlateElementVertexBuffer<FSlateVertex>& VertexBuffer, FSlateElementIndexBuffer& IndexBuffer);
-
 private:
 	/**
 	 * Returns the pixel shader that should be used for the specified ShaderType and DrawEffects
@@ -280,6 +279,11 @@ private:
 
 	/** @return The RHI primitive type from the Slate primitive type */
 	EPrimitiveType GetRHIPrimitiveType(ESlateDrawPrimitive::Type SlateType);
+
+	/**
+	 * Delete resources that have had their rendering resources released
+	 */
+	void DeleteReleasedResources();
 private:
 	/** Buffers used for rendering */
 	TSlateElementVertexBuffer<FSlateVertex> VertexBuffers[SlateRHIConstants::NumBuffers];
