Merged CL 2671863 from Orion

Fix the number of draw buffers for slate elements being off by one by mistake

diff --git a/Engine/Source/Runtime/SlateRHIRenderer/Private/SlateRHIRenderingPolicy.h b/Engine/Source/Runtime/SlateRHIRenderer/Private/SlateRHIRenderingPolicy.h
--- a/Engine/Source/Runtime/SlateRHIRenderer/Private/SlateRHIRenderingPolicy.h
+++ b/Engine/Source/Runtime/SlateRHIRenderer/Private/SlateRHIRenderingPolicy.h
@@ -10,7 +10,7 @@ class FSlateRHIResourceManager;
 namespace SlateRHIConstants
 {
 	// Number of vertex and index buffers we swap between when drawing windows
-	const int32 NumBuffers = 2;
+	const int32 NumBuffers = 3;
 }
 
 /** 
