Merged CL 2672323 from Orion

Fix multithreaded rendering crashes due to race conditions in the Slate renderer

diff --git a/Engine/Source/Runtime/SlateCore/Private/Rendering/DrawElements.cpp b/Engine/Source/Runtime/SlateCore/Private/Rendering/DrawElements.cpp
--- a/Engine/Source/Runtime/SlateCore/Private/Rendering/DrawElements.cpp
+++ b/Engine/Source/Runtime/SlateCore/Private/Rendering/DrawElements.cpp
@@ -593,9 +593,9 @@ TSharedRef<FSlateRenderDataHandle, ESPMode::ThreadSafe> FSlateWindowElementList:
 	return CachedRenderDataHandleRef;
 }
 
-void FSlateWindowElementList::PreDraw_RenderThread()
+void FSlateWindowElementList::PreDraw_ParallelThread()
 {
-	check(IsInRenderingThread());
+	check(IsInParallelRenderingThread());
 
 	for ( auto& Entry : DrawLayers )
 	{
@@ -604,9 +604,9 @@ void FSlateWindowElementList::PreDraw_RenderThread()
 	}
 }
 
-void FSlateWindowElementList::PostDraw_RenderThread()
+void FSlateWindowElementList::PostDraw_ParallelThread()
 {
-	check(IsInRenderingThread());
+	check(IsInParallelRenderingThread());
 
 	for ( auto& Entry : DrawLayers )
 	{
diff --git a/Engine/Source/Runtime/SlateCore/Public/Rendering/DrawElements.h b/Engine/Source/Runtime/SlateCore/Public/Rendering/DrawElements.h
--- a/Engine/Source/Runtime/SlateCore/Public/Rendering/DrawElements.h
+++ b/Engine/Source/Runtime/SlateCore/Public/Rendering/DrawElements.h
@@ -992,9 +992,9 @@ public:
 		return SafeHandle.IsValid() && SafeHandle->IsInUse();
 	}
 
-	SLATECORE_API void PreDraw_RenderThread();
+	SLATECORE_API void PreDraw_ParallelThread();
 
-	SLATECORE_API void PostDraw_RenderThread();
+	SLATECORE_API void PostDraw_ParallelThread();
 
 private:
 
diff --git a/Engine/Source/Runtime/SlateRHIRenderer/Private/SlateRHIRenderer.cpp b/Engine/Source/Runtime/SlateRHIRenderer/Private/SlateRHIRenderer.cpp
--- a/Engine/Source/Runtime/SlateRHIRenderer/Private/SlateRHIRenderer.cpp
+++ b/Engine/Source/Runtime/SlateRHIRenderer/Private/SlateRHIRenderer.cpp
@@ -202,7 +202,6 @@ void FSlateRHIRenderer::Destroy()
 /** Returns a draw buffer that can be used by Slate windows to draw window elements */
 FSlateDrawBuffer& FSlateRHIRenderer::GetDrawBuffer()
 {
-#if USE_MAX_DRAWBUFFERS
 	FreeBufferIndex = (FreeBufferIndex + 1) % NumDrawBuffers;
 	
 	FSlateDrawBuffer* Buffer = &DrawBuffers[FreeBufferIndex];
@@ -234,10 +233,6 @@ FSlateDrawBuffer& FSlateRHIRenderer::GetDrawBuffer()
 
 	Buffer->ClearBuffer();
 	return *Buffer;
-#else
-	// With this method buffers are created on this thread and deleted on the rendering thread
-	return *(new FSlateDrawBuffer);
-#endif
 }
 
 void FSlateRHIRenderer::CreateViewport( const TSharedRef<SWindow> Window )
@@ -407,7 +402,7 @@ void FSlateRHIRenderer::DrawWindow_RenderThread(FRHICommandListImmediate& RHICmd
 		FSlateBatchData& BatchData = WindowElementList.GetBatchData();
 		FElementBatchMap& RootBatchMap = WindowElementList.GetRootDrawLayer().GetElementBatchMap();
 
-		WindowElementList.PreDraw_RenderThread();
+		WindowElementList.PreDraw_ParallelThread();
 
 		{
 			SCOPE_CYCLE_COUNTER(STAT_SlateRTCreateBatches);
@@ -466,8 +461,6 @@ void FSlateRHIRenderer::DrawWindow_RenderThread(FRHICommandListImmediate& RHICmd
 				BatchData.GetRenderBatches()
 			);
 		}
-
-		WindowElementList.PostDraw_RenderThread();
 	}
 
 	bool bNeedCallFinishFrameForStereo = false;
@@ -539,15 +532,31 @@ void FSlateRHIRenderer::DrawWindows()
 	}
 }
 
-static void EndDrawingWindows( FSlateDrawBuffer* DrawBuffer, FSlateRHIRenderingPolicy& Policy )
+struct FSlateEndDrawingWindowsCommand : public FRHICommand<FSlateEndDrawingWindowsCommand>
 {
-#if USE_MAX_DRAWBUFFERS
-	DrawBuffer->Unlock();
-#else
-	delete DrawBuffer;
-#endif
+	FSlateRHIRenderingPolicy& Policy;
+	FSlateDrawBuffer* DrawBuffer;
+	
+	FSlateEndDrawingWindowsCommand(FSlateRHIRenderingPolicy& InPolicy, FSlateDrawBuffer* InDrawBuffer)
+		: Policy( InPolicy )
+		, DrawBuffer(InDrawBuffer)
+	{}
 
-	Policy.EndDrawingWindows();
+	void Execute(FRHICommandListBase& CmdList)
+	{
+		for( auto& ElementList : DrawBuffer->GetWindowElementLists() )
+		{
+			ElementList->PostDraw_ParallelThread();
+		}
+
+		DrawBuffer->Unlock();
+		Policy.EndDrawingWindows();
+	}
+};
+
+static void EndDrawingWindows( FRHICommandListImmediate& RHICmdList, FSlateDrawBuffer* DrawBuffer, FSlateRHIRenderingPolicy& Policy )
+{
+	new (RHICmdList.AllocCommand<FSlateEndDrawingWindowsCommand>()) FSlateEndDrawingWindowsCommand(Policy, DrawBuffer);
 }
 
 
@@ -695,7 +704,7 @@ void FSlateRHIRenderer::DrawWindows_Private( FSlateDrawBuffer& WindowDrawBuffer
 		FSlateDrawBuffer*, DrawBuffer, &WindowDrawBuffer,
 		FSlateRHIRenderingPolicy&, Policy, *RenderingPolicy,
 	{
-		EndDrawingWindows( DrawBuffer, Policy );
+		EndDrawingWindows( RHICmdList, DrawBuffer, Policy );
 	});
 
 	// flush the cache if needed
diff --git a/Engine/Source/Runtime/SlateRHIRenderer/Private/SlateRHIRenderer.h b/Engine/Source/Runtime/SlateRHIRenderer/Private/SlateRHIRenderer.h
--- a/Engine/Source/Runtime/SlateRHIRenderer/Private/SlateRHIRenderer.h
+++ b/Engine/Source/Runtime/SlateRHIRenderer/Private/SlateRHIRenderer.h
@@ -7,12 +7,8 @@ class FSlateRHIRenderingPolicy;
 class FSlateElementBatcher;
 
 
-#define USE_MAX_DRAWBUFFERS 1
-
-#if USE_MAX_DRAWBUFFERS
 // Number of draw buffers that can be active at any given time
 const uint32 NumDrawBuffers = 3;
-#endif
 
 // Enable to visualize overdraw in Slate
 #define DEBUG_OVERDRAW 0
@@ -269,13 +265,11 @@ private:
 	/** Keep a pointer around for when we have deferred drawing happening */
 	FSlateDrawBuffer* EnqueuedWindowDrawBuffer;
 
-#if USE_MAX_DRAWBUFFERS
 	/** Double buffered draw buffers so that the rendering thread can be rendering windows while the game thread is setting up for next frame */
 	FSlateDrawBuffer DrawBuffers[NumDrawBuffers];
 
 	/** The draw buffer which is currently free for use by the game thread */
 	uint8 FreeBufferIndex;
-#endif
 
 	/** Element batcher which renders draw elements */
 	TSharedPtr<FSlateElementBatcher> ElementBatcher;
diff --git a/Engine/Source/Runtime/SlateRHIRenderer/Private/SlateRHIRenderingPolicy.cpp b/Engine/Source/Runtime/SlateRHIRenderer/Private/SlateRHIRenderingPolicy.cpp
--- a/Engine/Source/Runtime/SlateRHIRenderer/Private/SlateRHIRenderingPolicy.cpp
+++ b/Engine/Source/Runtime/SlateRHIRenderer/Private/SlateRHIRenderingPolicy.cpp
@@ -243,7 +243,7 @@ void FSlateRHIRenderingPolicy::BeginDrawingWindows()
 
 void FSlateRHIRenderingPolicy::EndDrawingWindows()
 {
-	check( IsInRenderingThread() );
+	check( IsInParallelRenderingThread() );
 
 #if STATS
 	uint32 TotalVertexBufferMemory = 0;
