Merging using UE4_to_UE4-Orion

CL 2678623
Slate rendering resource access fixes.
- Removed the need for a render command per resource access.  This was unsafe when accessing resources from another thread (loading screens)
- Moved texture streaming last render time update to the actual point in time where it was rendered.  The previous location was at resource access time which was not occuring for textures in invalidation panels.

diff --git a/Engine/Source/Runtime/SlateCore/Public/Textures/SlateShaderResource.h b/Engine/Source/Runtime/SlateCore/Public/Textures/SlateShaderResource.h
--- a/Engine/Source/Runtime/SlateCore/Public/Textures/SlateShaderResource.h
+++ b/Engine/Source/Runtime/SlateCore/Public/Textures/SlateShaderResource.h
@@ -18,7 +18,10 @@ namespace ESlateShaderResource
 		TextureObject,
 
 		/** Material resource. */
-		Material
+		Material,
+
+		/** No Resource. */
+		Invalid,
 	};
 }
 
diff --git a/Engine/Source/Runtime/SlateRHIRenderer/Private/SlateRHIRenderingPolicy.cpp b/Engine/Source/Runtime/SlateRHIRenderer/Private/SlateRHIRenderingPolicy.cpp
--- a/Engine/Source/Runtime/SlateRHIRenderer/Private/SlateRHIRenderingPolicy.cpp
+++ b/Engine/Source/Runtime/SlateRHIRenderer/Private/SlateRHIRenderingPolicy.cpp
@@ -599,7 +599,6 @@ void FSlateRHIRenderingPolicy::DrawElements(FRHICommandListImmediate& RHICmdList
 		}
 
 		const FSlateShaderResource* ShaderResource = RenderBatch.Texture;
-
 		const ESlateBatchDrawFlag::Type DrawFlags = RenderBatch.DrawFlags;
 		const ESlateDrawEffect::Type DrawEffects = RenderBatch.DrawEffects;
 		const ESlateShader::Type ShaderType = RenderBatch.ShaderType;
@@ -618,7 +617,8 @@ void FSlateRHIRenderingPolicy::DrawElements(FRHICommandListImmediate& RHICmdList
 				RHICmdList.SetScissorRect(false, 0, 0, 0, 0);
 			}
 
-			if( !ShaderResource || ShaderResource->GetType() != ESlateShaderResource::Material ) 
+			ESlateShaderResource::Type ResourceType = ShaderResource ? ShaderResource->GetType() : ESlateShaderResource::Invalid;
+			if( ResourceType != ESlateShaderResource::Material ) 
 			{
 				FSlateElementPS* PixelShader = GetTexturePixelShader(ShaderType, DrawEffects);
 
@@ -656,17 +656,21 @@ void FSlateRHIRenderingPolicy::DrawElements(FRHICommandListImmediate& RHICmdList
 					RHICmdList.SetRasterizerState(TStaticRasterizerState<FM_Solid, CM_None, true>::GetRHI());
 				}
 
-			
 
-				FTexture2DRHIRef TextureRHI;
-				if(ShaderResource)
+				FSamplerStateRHIRef SamplerState = BilinearClamp;
+				FTextureRHIRef TextureRHI = GWhiteTexture->TextureRHI;
+				if( ShaderResource )
 				{
-					TextureRHI = ((TSlateTexture<FTexture2DRHIRef>*)ShaderResource)->GetTypedResource();
-				}
-
-				if( ShaderResource && IsValidRef( TextureRHI ) )
-				{
-					FSamplerStateRHIRef SamplerState; 
+					if (ResourceType == ESlateShaderResource::TextureObject)
+					{
+						FSlateUTextureResource* TextureObjectResource = (FSlateUTextureResource*)ShaderResource;
+		
+						TextureRHI = TextureObjectResource->AccessRHIResource();
+					}
+					else
+					{	
+						TextureRHI = ((TSlateTexture<FTexture2DRHIRef>*)ShaderResource)->GetTypedResource();
+					}
 
 					if( DrawFlags == (ESlateBatchDrawFlag::TileU | ESlateBatchDrawFlag::TileV) )
 					{
@@ -684,13 +688,9 @@ void FSlateRHIRenderingPolicy::DrawElements(FRHICommandListImmediate& RHICmdList
 					{
 						SamplerState = BilinearClamp;
 					}
-					PixelShader->SetTexture(RHICmdList, TextureRHI, SamplerState);
-				}
-				else
-				{
-					PixelShader->SetTexture(RHICmdList, GWhiteTexture->TextureRHI, BilinearClamp);
 				}
 
+				PixelShader->SetTexture(RHICmdList, TextureRHI, SamplerState);
 				PixelShader->SetShaderParams(RHICmdList, ShaderParams.PixelParams);
 				PixelShader->SetDisplayGamma(RHICmdList, (DrawFlags & ESlateBatchDrawFlag::NoGamma) ? 1.0f : DisplayGamma);
 
diff --git a/Engine/Source/Runtime/SlateRHIRenderer/Private/SlateRHIResourceManager.cpp b/Engine/Source/Runtime/SlateRHIRenderer/Private/SlateRHIResourceManager.cpp
--- a/Engine/Source/Runtime/SlateRHIRenderer/Private/SlateRHIResourceManager.cpp
+++ b/Engine/Source/Runtime/SlateRHIRenderer/Private/SlateRHIResourceManager.cpp
@@ -613,29 +613,10 @@ FSlateShaderResourceProxy* FSlateRHIResourceManager::FindOrCreateDynamicTextureR
 				}
 			}
 
-			if ( TextureResource.IsValid() )
+			if ( TextureResource.IsValid() && TextureResource->TextureObject && TextureResource->TextureObject->Resource )
 			{
-				UTexture* Texture2DObject = TextureResource->TextureObject;
-				if ( Texture2DObject && !AccessedUTextures.Contains(Texture2DObject) && Texture2DObject->Resource )
-				{
-					// Set the texture rendering resource that should be used.  The UTexture resource could change at any time so we must do this each frame
-					ENQUEUE_UNIQUE_RENDER_COMMAND_ONEPARAMETER(UpdateSlateUTextureResource,
-						FSlateUTextureResource*, InUTextureResource, TextureResource.Get(),
-						{
-							FTexture* RenderTexture = InUTextureResource->TextureObject->Resource;
-
-							// Let the streaming manager know we are using this texture now
-							RenderTexture->LastRenderTime = FApp::GetCurrentTime();
-
-							// Refresh FTexture
-							InUTextureResource->UpdateRenderResource(RenderTexture);
-
-						});
-
-
-					AccessedUTextures.Add(Texture2DObject);
-				}
-
+				TextureResource->UpdateRenderResource(TextureObject->Resource);
+				AccessedUTextures.Add(TextureResource->TextureObject);
 				return TextureResource->Proxy;
 			}
 		}
diff --git a/Engine/Source/Runtime/SlateRHIRenderer/Private/SlateUTextureResource.cpp b/Engine/Source/Runtime/SlateRHIRenderer/Private/SlateUTextureResource.cpp
--- a/Engine/Source/Runtime/SlateRHIRenderer/Private/SlateUTextureResource.cpp
+++ b/Engine/Source/Runtime/SlateRHIRenderer/Private/SlateUTextureResource.cpp
@@ -32,15 +32,10 @@ void FSlateUTextureResource::UpdateRenderResource(FTexture* InFTexture)
 		// if that's true we need to update the actual size recorded on the proxy as well, otherwise 
 		// the texture will continue to render using the wrong size.
 		Proxy->ActualSize = FIntPoint(InFTexture->GetSizeX(), InFTexture->GetSizeY());
-
-		ShaderResource = FTexture2DRHIRef(InFTexture->TextureRHI->GetTexture2D());
-	}
-	else
-	{
-		ShaderResource = FTexture2DRHIRef();
 	}
 }
 
+
 uint32 FSlateUTextureResource::GetWidth() const
 { 
 	return TextureObject->GetSurfaceWidth();
diff --git a/Engine/Source/Runtime/SlateRHIRenderer/Private/SlateUTextureResource.h b/Engine/Source/Runtime/SlateRHIRenderer/Private/SlateUTextureResource.h
--- a/Engine/Source/Runtime/SlateRHIRenderer/Private/SlateUTextureResource.h
+++ b/Engine/Source/Runtime/SlateRHIRenderer/Private/SlateUTextureResource.h
@@ -3,7 +3,7 @@
 #pragma once
 
 /** A resource for rendering a UTexture object in Slate */
-class FSlateUTextureResource : public TSlateTexture<FTexture2DRHIRef>
+class FSlateUTextureResource : public FSlateShaderResource
 {
 public:
 	static TSharedPtr<FSlateUTextureResource> NullResource;
@@ -12,7 +12,7 @@ public:
 	~FSlateUTextureResource();
 
 	/**
-	 * Updates the renderng resource with a new UTexture resource
+	 * Updates the rendering resource with a potentially new texture
 	 */
 	void UpdateRenderResource(FTexture* InFTexture);
 
@@ -23,6 +23,21 @@ public:
 	virtual uint32 GetHeight() const;
 	virtual ESlateShaderResource::Type GetType() const;
 	
+	/** Gets the RHI resource used for rendering and updates the last render time for texture streaming */
+	FTextureRHIRef AccessRHIResource() 
+	{
+		if (TextureObject && TextureObject->Resource)
+		{
+			FTexture* TextureResource = TextureObject->Resource;
+			TextureResource->LastRenderTime = FApp::GetCurrentTime();
+
+			return TextureResource->TextureRHI;
+		}
+
+		return FTextureRHIRef();
+	}
+		
+	
 public:
 	/** Slate rendering proxy */
 	FSlateShaderResourceProxy* Proxy;
