Merging CL 2684013 and 2684020 using UE4_to_UE4-Orion

Added support for vertex shaders in UI materials.  Uses "WorldPositionOffset" like regular materials.  Stubbed in support for instanced parameters although that is not hooked up to anything yet

diff --git a/Engine/Shaders/SlateVertexShader.usf b/Engine/Shaders/SlateVertexShader.usf
--- a/Engine/Shaders/SlateVertexShader.usf
+++ b/Engine/Shaders/SlateVertexShader.usf
@@ -4,6 +4,10 @@
 #include "Common.usf"
 #include "GammaCorrectionCommon.usf"
 
+#if USE_MATERIALS
+#include "Material.usf"
+#endif
+
 float4x4 ViewProjection;
 
 float SwitchVerticalAxisMultiplier;
@@ -29,7 +33,19 @@ VertexOut Main(
 	)
 {
 	VertexOut VOut;
-	VOut.Position = mul(float4(InPosition.xy,0,1), ViewProjection);
+
+	float4 WorldPosition = float4(InPosition.xy,0,1);
+#if USE_MATERIALS
+	FMaterialVertexParameters VertexParameters;
+	VertexParameters.WorldPosition = WorldPosition.xyz;
+#if USE_INSTANCING
+	VertexParameters.PerInstanceParams = float4(0,0,0,0);
+#endif
+	VertexParameters.VertexColor = InColor;
+	WorldPosition.xyz += GetMaterialWorldPositionOffset(VertexParameters);
+#endif
+
+	VOut.Position = mul(WorldPosition, ViewProjection);
 
 #if ES2_PROFILE && COMPILER_GLSL_ES2
 	// @todo-mobile: Fix this in the projection matrix
diff --git a/Engine/Source/Editor/UnrealEd/Classes/MaterialGraph/MaterialGraph.h b/Engine/Source/Editor/UnrealEd/Classes/MaterialGraph/MaterialGraph.h
--- a/Engine/Source/Editor/UnrealEd/Classes/MaterialGraph/MaterialGraph.h
+++ b/Engine/Source/Editor/UnrealEd/Classes/MaterialGraph/MaterialGraph.h
@@ -46,7 +46,7 @@ struct FMaterialInputInfo
 		}
 		else if( Material->IsUIMaterial() )
 		{
-			return Property == MP_EmissiveColor || Property == MP_Opacity || Property == MP_OpacityMask;
+			return Property == MP_EmissiveColor || Property == MP_Opacity || Property == MP_OpacityMask || Property == MP_WorldPositionOffset;
 		}
 		else
 		{
diff --git a/Engine/Source/Runtime/Engine/Private/Materials/Material.cpp b/Engine/Source/Runtime/Engine/Private/Materials/Material.cpp
--- a/Engine/Source/Runtime/Engine/Private/Materials/Material.cpp
+++ b/Engine/Source/Runtime/Engine/Private/Materials/Material.cpp
@@ -3974,6 +3974,7 @@ bool UMaterial::IsPropertyActive(EMaterialProperty InProperty) const
 	else if ( MaterialDomain == MD_UI )
 	{
 		return InProperty == MP_EmissiveColor
+			|| ( InProperty == MP_WorldPositionOffset )
 			|| ( InProperty == MP_OpacityMask && BlendMode == BLEND_Masked ) 
 			|| ( InProperty == MP_Opacity && IsTranslucentBlendMode((EBlendMode)BlendMode) && BlendMode != BLEND_Modulate );
 	}
diff --git a/Engine/Source/Runtime/SlateRHIRenderer/Private/SlateMaterialShader.cpp b/Engine/Source/Runtime/SlateRHIRenderer/Private/SlateMaterialShader.cpp
--- a/Engine/Source/Runtime/SlateRHIRenderer/Private/SlateMaterialShader.cpp
+++ b/Engine/Source/Runtime/SlateRHIRenderer/Private/SlateMaterialShader.cpp
@@ -4,6 +4,58 @@
 #include "SlateMaterialShader.h"
 #include "RenderingCommon.h"
 
+
+FSlateMaterialShaderVS::FSlateMaterialShaderVS(const FMaterialShaderType::CompiledShaderInitializerType& Initializer)
+	: FMaterialShader(Initializer)
+{
+	ViewProjection.Bind(Initializer.ParameterMap, TEXT("ViewProjection"));
+	SwitchVerticalAxisMultiplier.Bind( Initializer.ParameterMap, TEXT("SwitchVerticalAxisMultiplier"));
+}
+
+
+void FSlateMaterialShaderVS::ModifyCompilationEnvironment(EShaderPlatform Platform, const FMaterial* Material, FShaderCompilerEnvironment& OutEnvironment)
+{
+	// Set defines based on what this shader will be used for
+	OutEnvironment.SetDefine( TEXT("USE_MATERIALS"), 1 );
+
+	FMaterialShader::ModifyCompilationEnvironment( Platform, Material, OutEnvironment );
+}
+
+bool FSlateMaterialShaderVS::ShouldCache(EShaderPlatform Platform, const FMaterial* Material)
+{
+	return Material->GetMaterialDomain() == MD_UI;
+}
+
+void FSlateMaterialShaderVS::SetViewProjection(FRHICommandList& RHICmdList, const FMatrix& InViewProjection )
+{
+	SetShaderValue(RHICmdList, GetVertexShader(), ViewProjection, InViewProjection );
+}
+
+void FSlateMaterialShaderVS::SetMaterialShaderParameters(FRHICommandList& RHICmdList, const FSceneView& View, const FMaterialRenderProxy* MaterialRenderProxy, const FMaterial* Material)
+{
+	const FVertexShaderRHIParamRef ShaderRHI = GetVertexShader();
+
+	const bool bDeferredPass = false;
+	FMaterialShader::SetParameters<FVertexShaderRHIParamRef>(RHICmdList, ShaderRHI, MaterialRenderProxy, *Material, View, bDeferredPass, ESceneRenderTargetsMode::DontSet);
+}
+
+void FSlateMaterialShaderVS::SetVerticalAxisMultiplier(FRHICommandList& RHICmdList, float InMultiplier )
+{
+	SetShaderValue(RHICmdList, GetVertexShader(), SwitchVerticalAxisMultiplier, InMultiplier );
+}
+
+
+bool FSlateMaterialShaderVS::Serialize(FArchive& Ar)
+{
+	bool bShaderHasOutdatedParameters = FMaterialShader::Serialize(Ar);
+
+	Ar << ViewProjection;
+	Ar << SwitchVerticalAxisMultiplier;
+
+	return bShaderHasOutdatedParameters;
+}
+
+
 bool FSlateMaterialShaderPS::ShouldCache(EShaderPlatform Platform, const FMaterial* Material)
 {
 	return Material->GetMaterialDomain() == MD_UI;
@@ -84,11 +136,21 @@ bool FSlateMaterialShaderPS::Serialize(FArchive& Ar)
 	return bShaderHasOutdatedParameters;
 }
 
+
+#define IMPLEMENT_SLATE_VERTEXMATERIALSHADER_TYPE(bUseInstancing) \
+	typedef TSlateMaterialShaderVS<bUseInstancing> TSlateMaterialShaderVS##bUseInstancing; \
+	IMPLEMENT_MATERIAL_SHADER_TYPE(template<>, TSlateMaterialShaderVS##bUseInstancing, TEXT("SlateVertexShader"), TEXT("Main"), SF_Vertex);
+
+/** Instancing vertex shader */
+IMPLEMENT_SLATE_VERTEXMATERIALSHADER_TYPE(true);
+
+/** Non instancing vertex shader */
+IMPLEMENT_SLATE_VERTEXMATERIALSHADER_TYPE(false);
+
 #define IMPLEMENT_SLATE_MATERIALSHADER_TYPE(ShaderType, bDrawDisabledEffect) \
 	typedef TSlateMaterialShaderPS<ESlateShader::ShaderType, bDrawDisabledEffect> TSlateMaterialShaderPS##ShaderType##bDrawDisabledEffect; \
 	IMPLEMENT_MATERIAL_SHADER_TYPE(template<>, TSlateMaterialShaderPS##ShaderType##bDrawDisabledEffect, TEXT("SlateElementPixelShader"), TEXT("Main"), SF_Pixel);
 
-
 IMPLEMENT_SLATE_MATERIALSHADER_TYPE(Default,true);
 IMPLEMENT_SLATE_MATERIALSHADER_TYPE(Default,false);
 IMPLEMENT_SLATE_MATERIALSHADER_TYPE(Border,true);
diff --git a/Engine/Source/Runtime/SlateRHIRenderer/Private/SlateMaterialShader.h b/Engine/Source/Runtime/SlateRHIRenderer/Private/SlateMaterialShader.h
--- a/Engine/Source/Runtime/SlateRHIRenderer/Private/SlateMaterialShader.h
+++ b/Engine/Source/Runtime/SlateRHIRenderer/Private/SlateMaterialShader.h
@@ -4,6 +4,38 @@
 
 #include "Runtime/Renderer/Public/MaterialShader.h"
 
+class FSlateMaterialShaderVS : public FMaterialShader
+{
+public:
+	FSlateMaterialShaderVS() {}
+	FSlateMaterialShaderVS(const FMaterialShaderType::CompiledShaderInitializerType& Initializer);
+
+	static void ModifyCompilationEnvironment(EShaderPlatform Platform, const FMaterial* Material, FShaderCompilerEnvironment& OutEnvironment);
+
+	static bool ShouldCache(EShaderPlatform Platform, const FMaterial* Material);
+
+	/** 
+	 * Sets the view projection parameter
+	 *
+	 * @param InViewProjection	The ViewProjection matrix to use when this shader is bound 
+	 */
+	void SetViewProjection(FRHICommandList& RHICmdList, const FMatrix& InViewProjection );
+
+	void SetMaterialShaderParameters(FRHICommandList& RHICmdList, const FSceneView& View, const FMaterialRenderProxy* MaterialRenderProxy, const FMaterial* Material);
+
+	/**
+	 * Sets the vertical axis multiplier to use depending on graphics api
+	 */
+	void SetVerticalAxisMultiplier(FRHICommandList& RHICmdList, float InMultiplier);
+
+	/** Serializes the shader data */
+	virtual bool Serialize( FArchive& Ar ) override;
+private:
+	/** ViewProjection parameter used by the shader */
+	FShaderParameter ViewProjection;
+	/** Parameter used to determine if we need to swtich the vertical axis for opengl */
+	FShaderParameter SwitchVerticalAxisMultiplier;
+};
 
 class FSlateMaterialShaderPS : public FMaterialShader
 {
@@ -34,6 +66,37 @@ private:
 	FShaderResourceParameter AdditionalTextureParameter;
 };
 
+template<bool bUseInstancing> 
+class TSlateMaterialShaderVS : public FSlateMaterialShaderVS
+{
+public:
+	DECLARE_SHADER_TYPE(TSlateMaterialShaderVS,Material);
+
+	TSlateMaterialShaderVS() { }
+
+	TSlateMaterialShaderVS(const ShaderMetaType::CompiledShaderInitializerType& Initializer)
+		: FSlateMaterialShaderVS( Initializer )
+	{ }
+	
+	/** Only compile shaders used with UI. */
+	static bool ShouldCache(EShaderPlatform Platform, const FMaterial* Material)
+	{
+		return FSlateMaterialShaderVS::ShouldCache(Platform,Material);
+	}
+
+	/** Modifies the compilation of this shader. */
+	static void ModifyCompilationEnvironment(EShaderPlatform Platform, const FMaterial* Material, FShaderCompilerEnvironment& OutEnvironment)
+	{
+		FSlateMaterialShaderVS::ModifyCompilationEnvironment(Platform, Material,OutEnvironment);
+
+		OutEnvironment.SetDefine(TEXT("USE_INSTANCING"), (uint32)(bUseInstancing ? 1 : 0));
+	}
+
+	virtual bool Serialize(FArchive& Ar) override
+	{
+		return FSlateMaterialShaderVS::Serialize( Ar );
+	}
+};
 
 template<ESlateShader::Type ShaderType,bool bDrawDisabledEffect> 
 class TSlateMaterialShaderPS : public FSlateMaterialShaderPS
diff --git a/Engine/Source/Runtime/SlateRHIRenderer/Private/SlateRHIRenderingPolicy.cpp b/Engine/Source/Runtime/SlateRHIRenderer/Private/SlateRHIRenderingPolicy.cpp
--- a/Engine/Source/Runtime/SlateRHIRenderer/Private/SlateRHIRenderingPolicy.cpp
+++ b/Engine/Source/Runtime/SlateRHIRenderer/Private/SlateRHIRenderingPolicy.cpp
@@ -547,7 +547,7 @@ void FSlateRHIRenderingPolicy::DrawElements(FRHICommandListImmediate& RHICmdList
 
 	FSceneView* SceneView = NULL;
 
-	TShaderMapRef<FSlateElementVS> VertexShader(GetGlobalShaderMap(GMaxRHIFeatureLevel));
+	TShaderMapRef<FSlateElementVS> GlobalVertexShader(GetGlobalShaderMap(GMaxRHIFeatureLevel));
 
 	// Disabled stencil test state
 	FDepthStencilStateRHIRef DSOff = TStaticDepthStencilState<false,CF_Always>::GetRHI();
@@ -608,6 +608,8 @@ void FSlateRHIRenderingPolicy::DrawElements(FRHICommandListImmediate& RHICmdList
 		{
 			check(RenderBatch.NumIndices > 0);
 
+			FMatrix DynamicOffset = FTranslationMatrix::Make(FVector(RenderBatch.DynamicOffset.X, RenderBatch.DynamicOffset.Y, 0));
+
 			if (RenderBatch.ScissorRect.IsSet())
 			{
 				RHICmdList.SetScissorRect(true, RenderBatch.ScissorRect.GetValue().Left, RenderBatch.ScissorRect.GetValue().Top, RenderBatch.ScissorRect.GetValue().Right, RenderBatch.ScissorRect.GetValue().Bottom);
@@ -624,16 +626,15 @@ void FSlateRHIRenderingPolicy::DrawElements(FRHICommandListImmediate& RHICmdList
 
 				RHICmdList.SetLocalBoundShaderState(RHICmdList.BuildLocalBoundShaderState(
 					GSlateVertexDeclaration.VertexDeclarationRHI,
-					VertexShader->GetVertexShader(),
+					GlobalVertexShader->GetVertexShader(),
 					nullptr,
 					nullptr,
 					PixelShader->GetPixelShader(),
 					FGeometryShaderRHIRef()));
 
-				FMatrix DynamicOffset = FTranslationMatrix::Make(FVector(RenderBatch.DynamicOffset.X, RenderBatch.DynamicOffset.Y, 0));
 
-				VertexShader->SetViewProjection(RHICmdList, DynamicOffset * ViewProjectionMatrix);
-				VertexShader->SetVerticalAxisMultiplier(RHICmdList, bAllowSwitchVerticalAxis && RHINeedsToSwitchVerticalAxis(GShaderPlatformForFeatureLevel[GMaxRHIFeatureLevel]) ? -1.0f : 1.0f );
+				GlobalVertexShader->SetViewProjection(RHICmdList, DynamicOffset * ViewProjectionMatrix);
+				GlobalVertexShader->SetVerticalAxisMultiplier(RHICmdList, bAllowSwitchVerticalAxis && RHINeedsToSwitchVerticalAxis(GShaderPlatformForFeatureLevel[GMaxRHIFeatureLevel]) ? -1.0f : 1.0f );
 #if !DEBUG_OVERDRAW
 				RHICmdList.SetBlendState(
 					(RenderBatch.DrawFlags & ESlateBatchDrawFlag::NoBlending)
@@ -723,7 +724,10 @@ void FSlateRHIRenderingPolicy::DrawElements(FRHICommandListImmediate& RHICmdList
 
 				FSlateMaterialShaderPS* PixelShader = GetMaterialPixelShader( Material, ShaderType, DrawEffects );
 
-				if( PixelShader )
+				const bool bUseInstancing = false;
+				FSlateMaterialShaderVS* VertexShader = GetMaterialVertexShader( Material, bUseInstancing );
+
+				if( VertexShader && PixelShader )
 				{
 					RHICmdList.SetLocalBoundShaderState(RHICmdList.BuildLocalBoundShaderState(
 						GSlateVertexDeclaration.VertexDeclarationRHI,
@@ -733,6 +737,10 @@ void FSlateRHIRenderingPolicy::DrawElements(FRHICommandListImmediate& RHICmdList
 						PixelShader->GetPixelShader(),
 						FGeometryShaderRHIRef()));
 
+					VertexShader->SetViewProjection(RHICmdList, DynamicOffset * ViewProjectionMatrix);
+					VertexShader->SetVerticalAxisMultiplier(RHICmdList, bAllowSwitchVerticalAxis && RHINeedsToSwitchVerticalAxis(GShaderPlatformForFeatureLevel[GMaxRHIFeatureLevel]) ? -1.0f : 1.0f );
+					VertexShader->SetMaterialShaderParameters(RHICmdList, *SceneView, MaterialRenderProxy, Material);
+
 					PixelShader->SetParameters(RHICmdList, *SceneView, MaterialRenderProxy, Material, DisplayGamma, ShaderParams.PixelParams);
 					PixelShader->SetDisplayGamma(RHICmdList, (DrawFlags & ESlateBatchDrawFlag::NoGamma) ? 1.0f : DisplayGamma);
 					
@@ -744,13 +752,9 @@ void FSlateRHIRenderingPolicy::DrawElements(FRHICommandListImmediate& RHICmdList
 						FTexture2DRHIRef TextureRHI;
 						TextureRHI = ((TSlateTexture<FTexture2DRHIRef>*)MaskResource)->GetTypedResource();
 
-			
 						PixelShader->SetAdditionalTexture(RHICmdList, TextureRHI, BilinearClamp);
 					}
 
-					VertexShader->SetViewProjection( RHICmdList, ViewProjectionMatrix );
-
-
 					uint32 PrimitiveCount = RenderBatch.DrawPrimitiveType == ESlateDrawPrimitive::LineList ? RenderBatch.NumIndices / 2 : RenderBatch.NumIndices / 3;
 
 					// for RHIs that can't handle VertexOffset, we need to offset the stream source each time
@@ -918,6 +922,23 @@ FSlateMaterialShaderPS* FSlateRHIRenderingPolicy::GetMaterialPixelShader( const
 	return FoundShader ? (FSlateMaterialShaderPS*)FoundShader->GetShaderChecked() : nullptr;
 }
 
+FSlateMaterialShaderVS* FSlateRHIRenderingPolicy::GetMaterialVertexShader( const FMaterial* Material, bool bUseInstancing )
+{
+	const FMaterialShaderMap* MaterialShaderMap = Material->GetRenderingThreadShaderMap();
+
+	FShader* FoundShader = nullptr;
+	if( bUseInstancing )
+	{
+		FoundShader = MaterialShaderMap->GetShader(&TSlateMaterialShaderVS<true>::StaticType);
+	}
+	else
+	{
+		FoundShader = MaterialShaderMap->GetShader(&TSlateMaterialShaderVS<false>::StaticType);
+	}
+	
+	return FoundShader ? (FSlateMaterialShaderVS*)FoundShader->GetShaderChecked() : nullptr;
+}
+
 EPrimitiveType FSlateRHIRenderingPolicy::GetRHIPrimitiveType(ESlateDrawPrimitive::Type SlateType)
 {
 	switch(SlateType)
diff --git a/Engine/Source/Runtime/SlateRHIRenderer/Private/SlateRHIRenderingPolicy.h b/Engine/Source/Runtime/SlateRHIRenderer/Private/SlateRHIRenderingPolicy.h
--- a/Engine/Source/Runtime/SlateRHIRenderer/Private/SlateRHIRenderingPolicy.h
+++ b/Engine/Source/Runtime/SlateRHIRenderer/Private/SlateRHIRenderingPolicy.h
@@ -276,6 +276,7 @@ private:
 	 */
 	class FSlateElementPS* GetTexturePixelShader( ESlateShader::Type ShaderType, ESlateDrawEffect::Type DrawEffects );
 	class FSlateMaterialShaderPS* GetMaterialPixelShader( const class FMaterial* Material, ESlateShader::Type ShaderType, ESlateDrawEffect::Type DrawEffects );
+	class FSlateMaterialShaderVS* GetMaterialVertexShader( const class FMaterial* Material, bool bUseInstancing );
 
 	/** @return The RHI primitive type from the Slate primitive type */
 	EPrimitiveType GetRHIPrimitiveType(ESlateDrawPrimitive::Type SlateType);
