Fixed crash in FSlate3DRenderer::DrawWindow_GameThread when the font cache was flushed

The font cache could be flushed while the render thread was still using the font atlas textures. We now flush it before we add the elements to draw to ensure that the atlas textures will be valid for the render thread to use.

diff --git a/Engine/Source/Runtime/SlateRHIRenderer/Private/Slate3DRenderer.cpp b/Engine/Source/Runtime/SlateRHIRenderer/Private/Slate3DRenderer.cpp
--- a/Engine/Source/Runtime/SlateRHIRenderer/Private/Slate3DRenderer.cpp
+++ b/Engine/Source/Runtime/SlateRHIRenderer/Private/Slate3DRenderer.cpp
@@ -51,6 +51,10 @@ void FSlate3DRenderer::DrawWindow_GameThread(FSlateDrawBuffer& DrawBuffer)
 {
 	check( IsInGameThread() );
 
+	// Need to flush the font cache before we add the elements below to avoid the flush potentially 
+	// deleting the texture resources that will be needed by the render thread
+	FontCache->ConditionalFlushCache();
+
 	TArray<TSharedPtr<FSlateWindowElementList>>& WindowElementLists = DrawBuffer.GetWindowElementLists();
 
 	// Only one window element list for now.
@@ -71,15 +75,10 @@ void FSlate3DRenderer::DrawWindow_GameThread(FSlateDrawBuffer& DrawBuffer)
 			// Update the font cache with new text after elements are batched
 			FontCache->UpdateCache();
 
-			bool temp = false;
-
-
 			// All elements for this window have been batched and rendering data updated
 			ElementBatcher->ResetBatches();
 		}
 	}
-
-	FontCache->ConditionalFlushCache();
 }
 
 void FSlate3DRenderer::DrawWindowToTarget_RenderThread( FRHICommandListImmediate& RHICmdList, UTextureRenderTarget2D* RenderTarget, FSlateDrawBuffer& InDrawBuffer )
