Fix CL #2702631 "Merged UE4-Orion->UE4. Slate hardware instancing supportΓÇ¥  for Metal by:
-  USE_SLATE_INSTANCING in the .usf as well as SlateMaterialShader.cpp as suggested.
- Separate instanced & non-instanced vertex declaration (or again weΓÇÖll get errors in Metal).
- Select the right declaration & shader in the rendering policy based on instancing.
- Bind the instance buffer only when required - removed branches when not rendering a material as that isn't instanced yet and added branches in !GRHISupportsBaseVertexIndex path because AMDΓÇÖs Metal drivers donΓÇÖt handle BaseVertexIndex correctly so it has to be switched off on *AMD only* under Metal.

diff --git a/Engine/Shaders/SlateVertexShader.usf b/Engine/Shaders/SlateVertexShader.usf
--- a/Engine/Shaders/SlateVertexShader.usf
+++ b/Engine/Shaders/SlateVertexShader.usf
@@ -15,8 +15,10 @@ VertexToPixelInterpolants Main(
 	in float2 InPosition : ATTRIBUTE2,
 	in float2 InClipOrigin : ATTRIBUTE3,
 	in float4 InClipExtents : ATTRIBUTE4,
-	in float4 InColor : ATTRIBUTE5,
-	in float4 InstanceParam : ATTRIBUTE6
+	in float4 InColor : ATTRIBUTE5
+#if USE_SLATE_INSTANCING
+	, in float4 InstanceParam : ATTRIBUTE6
+#endif
 	)
 {
 	VertexToPixelInterpolants VOut = (VertexToPixelInterpolants)0;
@@ -40,7 +42,10 @@ VertexToPixelInterpolants Main(
 
 	#if NUM_MATERIAL_TEXCOORDS_VERTEX
 
-		float2 UVArrays[8] = 
+#if !USE_SLATE_INSTANCING
+		float4 InstanceParam = float4(0, 0, 0, 0);
+#endif
+		float2 UVArrays[8] =
 		{
 			InTextureCoordinates.xy,
 			InTextureCoordinates.zw,
diff --git a/Engine/Source/Runtime/SlateRHIRenderer/Private/SlateRHIRenderingPolicy.cpp b/Engine/Source/Runtime/SlateRHIRenderer/Private/SlateRHIRenderingPolicy.cpp
--- a/Engine/Source/Runtime/SlateRHIRenderer/Private/SlateRHIRenderingPolicy.cpp
+++ b/Engine/Source/Runtime/SlateRHIRenderer/Private/SlateRHIRenderingPolicy.cpp
@@ -625,6 +625,9 @@ void FSlateRHIRenderingPolicy::DrawElements(FRHICommandListImmediate& RHICmdList
 			{
 				FSlateElementPS* PixelShader = GetTexturePixelShader(ShaderType, DrawEffects);
 
+				const bool bUseInstancing = RenderBatch.InstanceCount > 1 && RenderBatch.InstanceData != nullptr;
+				check(bUseInstancing == false);
+				
 				RHICmdList.SetLocalBoundShaderState(RHICmdList.BuildLocalBoundShaderState(
 					GSlateVertexDeclaration.VertexDeclarationRHI,
 					GlobalVertexShader->GetVertexShader(),
@@ -706,16 +709,6 @@ void FSlateRHIRenderingPolicy::DrawElements(FRHICommandListImmediate& RHICmdList
 				}
 				else
 				{
-					if( RenderBatch.InstanceCount > 1 && RenderBatch.InstanceData != nullptr )
-					{
-						FSlateUpdatableInstanceBuffer* InstanceBuffer = (FSlateUpdatableInstanceBuffer*)RenderBatch.InstanceData;
-						InstanceBuffer->BindStreamSource( RHICmdList, 1 );
-					}
-					else
-					{
-						RHICmdList.SetStreamSource(1, nullptr, 0, 0);
-					}
-
 					RHICmdList.DrawIndexedPrimitive(IndexBuffer->IndexBufferRHI, GetRHIPrimitiveType(RenderBatch.DrawPrimitiveType), RenderBatch.VertexOffset, 0, RenderBatch.NumVertices, RenderBatch.IndexOffset, PrimitiveCount, RenderBatch.InstanceCount);
 				}
 
@@ -735,13 +728,13 @@ void FSlateRHIRenderingPolicy::DrawElements(FRHICommandListImmediate& RHICmdList
 
 				FSlateMaterialShaderPS* PixelShader = GetMaterialPixelShader( Material, ShaderType, DrawEffects );
 
-				const bool bUseInstancing = RenderBatch.InstanceCount > 1;
+				const bool bUseInstancing = RenderBatch.InstanceCount > 1 && RenderBatch.InstanceData != nullptr;
 				FSlateMaterialShaderVS* VertexShader = GetMaterialVertexShader( Material, bUseInstancing );
 
 				if( VertexShader && PixelShader )
 				{
 					RHICmdList.SetLocalBoundShaderState(RHICmdList.BuildLocalBoundShaderState(
-						GSlateVertexDeclaration.VertexDeclarationRHI,
+						bUseInstancing ? GSlateInstancedVertexDeclaration.VertexDeclarationRHI : GSlateVertexDeclaration.VertexDeclarationRHI,
 						VertexShader->GetVertexShader(),
 						nullptr,
 						nullptr,
@@ -771,12 +764,22 @@ void FSlateRHIRenderingPolicy::DrawElements(FRHICommandListImmediate& RHICmdList
 					// for RHIs that can't handle VertexOffset, we need to offset the stream source each time
 					if (!GRHISupportsBaseVertexIndex)
 					{
+						if( bUseInstancing )
+						{
+							FSlateUpdatableInstanceBuffer* InstanceBuffer = (FSlateUpdatableInstanceBuffer*)RenderBatch.InstanceData;
+							InstanceBuffer->BindStreamSource( RHICmdList, 1 );
+						}
+						else
+						{
+							RHICmdList.SetStreamSource(1, nullptr, 0, 0);
+						}
+						
 						RHICmdList.SetStreamSource(0, VertexBuffer->VertexBufferRHI, sizeof(FSlateVertex), RenderBatch.VertexOffset * sizeof(FSlateVertex));
 						RHICmdList.DrawIndexedPrimitive(IndexBuffer->IndexBufferRHI, GetRHIPrimitiveType(RenderBatch.DrawPrimitiveType), 0, 0, RenderBatch.NumVertices, RenderBatch.IndexOffset, PrimitiveCount, RenderBatch.InstanceCount);
 					}
 					else
 					{
-						if( bUseInstancing && RenderBatch.InstanceData != nullptr )
+						if( bUseInstancing )
 						{
 							FSlateUpdatableInstanceBuffer* InstanceBuffer = (FSlateUpdatableInstanceBuffer*)RenderBatch.InstanceData;
 							InstanceBuffer->BindStreamSource( RHICmdList, 1 );
diff --git a/Engine/Source/Runtime/SlateRHIRenderer/Private/SlateShaders.cpp b/Engine/Source/Runtime/SlateRHIRenderer/Private/SlateShaders.cpp
--- a/Engine/Source/Runtime/SlateRHIRenderer/Private/SlateShaders.cpp
+++ b/Engine/Source/Runtime/SlateRHIRenderer/Private/SlateShaders.cpp
@@ -34,6 +34,7 @@ IMPLEMENT_SLATE_PIXELSHADER_TYPE(LineSegment,	true, true);
 
 /** The Slate vertex declaration. */
 TGlobalResource<FSlateVertexDeclaration> GSlateVertexDeclaration;
+TGlobalResource<FSlateInstancedVertexDeclaration> GSlateInstancedVertexDeclaration;
 
 
 /************************************************************************/
@@ -49,9 +50,6 @@ void FSlateVertexDeclaration::InitRHI()
 	Elements.Add(FVertexElement(0, STRUCT_OFFSET(FSlateVertex, ClipRect) + STRUCT_OFFSET(FSlateRotatedClipRectType, TopLeft), VET_Float2, 3, Stride));
 	Elements.Add(FVertexElement(0, STRUCT_OFFSET(FSlateVertex, ClipRect) + STRUCT_OFFSET(FSlateRotatedClipRectType, ExtentX), VET_Float4, 4, Stride));
 	Elements.Add(FVertexElement(0, STRUCT_OFFSET(FSlateVertex, Color), VET_Color, 5, Stride));
-	
-	static const bool bUseInstanceIndex = true;
-	Elements.Add(FVertexElement(1, 0, VET_Float4, 6, sizeof(FVector4), bUseInstanceIndex));
 
 	VertexDeclarationRHI = RHICreateVertexDeclaration(Elements);
 }
@@ -61,6 +59,25 @@ void FSlateVertexDeclaration::ReleaseRHI()
 	VertexDeclarationRHI.SafeRelease();
 }
 
+
+/************************************************************************/
+/* FSlateInstancedVertexDeclaration                                     */
+/************************************************************************/
+void FSlateInstancedVertexDeclaration::InitRHI()
+{
+	FVertexDeclarationElementList Elements;
+	uint32 Stride = sizeof(FSlateVertex);
+	Elements.Add(FVertexElement(0, STRUCT_OFFSET(FSlateVertex, TexCoords), VET_Float4, 0, Stride));
+	Elements.Add(FVertexElement(0, STRUCT_OFFSET(FSlateVertex, MaterialTexCoords), VET_Float2, 1, Stride));
+	Elements.Add(FVertexElement(0, STRUCT_OFFSET(FSlateVertex, Position), VET_Float2, 2, Stride));
+	Elements.Add(FVertexElement(0, STRUCT_OFFSET(FSlateVertex, ClipRect) + STRUCT_OFFSET(FSlateRotatedClipRectType, TopLeft), VET_Float2, 3, Stride));
+	Elements.Add(FVertexElement(0, STRUCT_OFFSET(FSlateVertex, ClipRect) + STRUCT_OFFSET(FSlateRotatedClipRectType, ExtentX), VET_Float4, 4, Stride));
+	Elements.Add(FVertexElement(0, STRUCT_OFFSET(FSlateVertex, Color), VET_Color, 5, Stride));
+	Elements.Add(FVertexElement(1, 0, VET_Float4, 6, sizeof(FVector4), true));
+	
+	VertexDeclarationRHI = RHICreateVertexDeclaration(Elements);
+}
+
 void FSlateElementPS::ModifyCompilationEnvironment(EShaderPlatform Platform, FShaderCompilerEnvironment& OutEnvironment)
 {
 	static const auto CVar = IConsoleManager::Get().FindTConsoleVariableDataInt(TEXT("r.Tonemapper709"));
diff --git a/Engine/Source/Runtime/SlateRHIRenderer/Private/SlateShaders.h b/Engine/Source/Runtime/SlateRHIRenderer/Private/SlateShaders.h
--- a/Engine/Source/Runtime/SlateRHIRenderer/Private/SlateShaders.h
+++ b/Engine/Source/Runtime/SlateRHIRenderer/Private/SlateShaders.h
@@ -28,6 +28,18 @@ public:
 	virtual void ReleaseRHI() override;
 };
 
+/**
+ * The vertex declaration for the slate instanced vertex shader
+ */
+class FSlateInstancedVertexDeclaration : public FSlateVertexDeclaration
+{
+public:
+	virtual ~FSlateInstancedVertexDeclaration() {}
+	
+	/** Initializes the vertex declaration RHI resource */
+	virtual void InitRHI() override;
+};
+
 /** The slate Vertex shader representation */
 class FSlateElementVS : public FGlobalShader
 {
@@ -231,3 +243,6 @@ private:
 
 /** The simple element vertex declaration. */
 extern TGlobalResource<FSlateVertexDeclaration> GSlateVertexDeclaration;
+
+/** The instanced simple element vertex declaration. */
+extern TGlobalResource<FSlateInstancedVertexDeclaration> GSlateInstancedVertexDeclaration;
