Mergine engine changes to //UE4/Main @ 2748007

diff --git a/Engine/Source/Runtime/Renderer/Private/StaticMeshDrawList.inl b/Engine/Source/Runtime/Renderer/Private/StaticMeshDrawList.inl
--- a/Engine/Source/Runtime/Renderer/Private/StaticMeshDrawList.inl
+++ b/Engine/Source/Runtime/Renderer/Private/StaticMeshDrawList.inl
@@ -172,6 +172,8 @@ void TStaticMeshDrawList<DrawingPolicyType>::AddMesh(
 
 template<typename DrawingPolicyType>
 TStaticMeshDrawList<DrawingPolicyType>::TStaticMeshDrawList()
+	: FrameNumberForVisibleCount(MAX_uint32)
+	, ViewStateUniqueId(0)
 {
 	if(IsInRenderingThread())
 	{
@@ -267,18 +269,22 @@ bool TStaticMeshDrawList<DrawingPolicyType>::DrawVisibleInner(
 		const int32 NumElements = DrawingPolicyLink->Elements.Num();
 		FPlatformMisc::Prefetch(&DrawingPolicyLink->CompactElements.GetData()->MeshId);
 		const FElementCompact* CompactElementPtr = DrawingPolicyLink->CompactElements.GetData();
+		uint32 Count = 0;
 		for (int32 ElementIndex = 0; ElementIndex < NumElements; ElementIndex++, CompactElementPtr++)
 		{
 			if (StaticMeshVisibilityMap.AccessCorrespondingBit(FRelativeBitReference(CompactElementPtr->MeshId)))
 			{
 				const FElement& Element = DrawingPolicyLink->Elements[ElementIndex];
 				STAT(StatInc +=  Element.Mesh->GetNumPrimitives();)
+				int32 SubCount = Element.Mesh->Elements.Num();
+				Count += SubCount;
 				// Avoid the cache miss looking up batch visibility if there is only one element.
-				uint64 BatchElementMask = Element.Mesh->Elements.Num() == 1 ? 1 : BatchVisibilityArray[Element.Mesh->Id];
+				uint64 BatchElementMask = SubCount == 1 ? 1 : BatchVisibilityArray[Element.Mesh->Id];
 				DrawElement(RHICmdList, View, PolicyContext, Element, BatchElementMask, DrawingPolicyLink, bDrawnShared);
-				bDirty = true;
 			}
 		}
+		bDirty = bDirty || !!Count;
+		DrawingPolicyLink->VisibleCount = Count;
 	}
 	INC_DWORD_STAT_BY(STAT_StaticMeshTriangles, StatInc);
 	return bDirty;
@@ -399,55 +405,83 @@ void TStaticMeshDrawList<DrawingPolicyType>::DrawVisibleParallel(
 			TArray<uint16, SceneRenderingAllocator>& PerDrawingPolicyCounts = *new (FMemStack::Get()) TArray<uint16, SceneRenderingAllocator>();
 			PerDrawingPolicyCounts.AddZeroed(NumPolicies);
 
-			ParallelForWithPreWork(NumPolicies, 
-				[this, &PerDrawingPolicyCounts, &StaticMeshVisibilityMap, &BatchVisibilityArray](int32 Index)
+			bool bNeedScan = true;
+			bool bCountsAreAccurate = false;
+			uint32 ViewKey = ParallelCommandListSet.View.GetViewKey();
+			uint32 ViewFrame = ParallelCommandListSet.View.GetOcclusionFrameCounter();
+			if (ParallelCommandListSet.bBalanceCommandsWithLastFrame)
+			{
+				// this means this will never work in split screen or the editor (multiple views) I am ok with that. It is fixable with more elaborate bookkeeping and additional memory.
+				if (ViewKey && ViewStateUniqueId == ViewKey && FrameNumberForVisibleCount + 1 == ViewFrame)
 				{
-					int32 Count = 0;
-
-					FDrawingPolicyLink* DrawingPolicyLink = &DrawingPolicySet[OrderedDrawingPolicies[Index]];
-					const FElementCompact* CompactElementPtr = DrawingPolicyLink->CompactElements.GetData();
-					FPlatformMisc::Prefetch(CompactElementPtr);
-					const int32 NumElements = DrawingPolicyLink->CompactElements.Num();
-					FPlatformMisc::Prefetch(&DrawingPolicyLink->CompactElements.GetData()->MeshId);
-					for (int32 ElementIndex = 0; ElementIndex < NumElements; ElementIndex++, CompactElementPtr++)
+					QUICK_SCOPE_CYCLE_COUNTER(STAT_FDeferredShadingSceneRenderer_DrawVisibleParallel_LastFrameScan);
+					// could use a ParallelFor here, but would rather leave that perf for the game thread.
+					for (int32 Index = 0; Index < NumPolicies; Index++)
+					{
+						FDrawingPolicyLink* DrawingPolicyLink = &DrawingPolicySet[OrderedDrawingPolicies[Index]];
+						PerDrawingPolicyCounts[Index] = (uint16)FMath::Min<uint32>(DrawingPolicyLink->VisibleCount, MAX_uint16);
+					}
+					bNeedScan = false;
+				}
+			}
+
+			if (bNeedScan)
+			{
+				QUICK_SCOPE_CYCLE_COUNTER(STAT_FDeferredShadingSceneRenderer_DrawVisibleParallel_FullVisibilityScan);
+				ParallelForWithPreWork(NumPolicies, 
+					[this, &PerDrawingPolicyCounts, &StaticMeshVisibilityMap, &BatchVisibilityArray](int32 Index)
 					{
-						if (StaticMeshVisibilityMap.AccessCorrespondingBit(FRelativeBitReference(CompactElementPtr->MeshId)))
+						int32 Count = 0;
+
+						FDrawingPolicyLink* DrawingPolicyLink = &DrawingPolicySet[OrderedDrawingPolicies[Index]];
+						const FElementCompact* CompactElementPtr = DrawingPolicyLink->CompactElements.GetData();
+						FPlatformMisc::Prefetch(CompactElementPtr);
+						const int32 NumElements = DrawingPolicyLink->CompactElements.Num();
+						FPlatformMisc::Prefetch(&DrawingPolicyLink->CompactElements.GetData()->MeshId);
+						for (int32 ElementIndex = 0; ElementIndex < NumElements; ElementIndex++, CompactElementPtr++)
 						{
-							const FElement& Element = DrawingPolicyLink->Elements[ElementIndex];
-							Count += Element.Mesh->Elements.Num();
+							if (StaticMeshVisibilityMap.AccessCorrespondingBit(FRelativeBitReference(CompactElementPtr->MeshId)))
+							{
+								const FElement& Element = DrawingPolicyLink->Elements[ElementIndex];
+								Count += Element.Mesh->Elements.Num();
+							}
 						}
-					}
-					if (Count)
+						if (Count)
+						{
+							// this is unlikely to overflow, but I don't think it would matter much if it did
+							PerDrawingPolicyCounts[Index] = (uint16)FMath::Min<int32>(Count, MAX_uint16);
+						}
+						DrawingPolicyLink->VisibleCount = Count;
+					},
+					[]()
 					{
-						// this is unlikely to overflow, but I don't think it would matter much if it did
-						PerDrawingPolicyCounts[Index] = (uint16)FMath::Min<int32>(Count, MAX_uint16);
+						QUICK_SCOPE_CYCLE_COUNTER(STAT_FDeferredShadingSceneRenderer_DrawVisibleParallel_ServiceLocalQueue);
+						FTaskGraphInterface::Get().ProcessThreadUntilIdle(ENamedThreads::RenderThread_Local);
 					}
-				},
-				[]()
-				{
-					QUICK_SCOPE_CYCLE_COUNTER(STAT_FDeferredShadingSceneRenderer_DrawVisibleParallel_ServiceLocalQueue);
-					FTaskGraphInterface::Get().ProcessThreadUntilIdle(ENamedThreads::RenderThread_Local);
-				}
-			);
+				);
+				bCountsAreAccurate = true;
+			}
+			FrameNumberForVisibleCount = ViewFrame;
+			ViewStateUniqueId = ViewKey;
 
 			int32 Total = 0;
 			for (uint16 Count : PerDrawingPolicyCounts)
 			{
 				Total += Count;
 			}
-			if (Total == 0)
+			if (bCountsAreAccurate && Total == 0)
 			{
 				return;
 			}
 			UE_CLOG(ParallelCommandListSet.bSpewBalance, LogTemp, Display, TEXT("Total Draws %d"), Total);
 
-			EffectiveThreads = FMath::Min<int32>(EffectiveThreads, FMath::DivideAndRoundUp(Total, ParallelCommandListSet.MinDrawsPerCommandList));
+			EffectiveThreads = FMath::Min<int32>(EffectiveThreads, FMath::DivideAndRoundUp(FMath::Max<int32>(1, Total), ParallelCommandListSet.MinDrawsPerCommandList));
 			check(EffectiveThreads > 0 && EffectiveThreads <= ParallelCommandListSet.Width);
 
 			int32 DrawsPerCmdList = FMath::DivideAndRoundUp(Total, EffectiveThreads);
 			int32 DrawsPerCmdListMergeLimit = FMath::DivideAndRoundUp(DrawsPerCmdList, 3); // if the last list would be small, we just merge it into the previous one
 
-	int32 Start = 0;
+			int32 Start = 0;
 
 			int32 PreviousBatchStart = -1;
 			int32 PreviousBatchEnd = -2;
@@ -456,7 +490,7 @@ void TStaticMeshDrawList<DrawingPolicyType>::DrawVisibleParallel(
 			while (Start < NumPolicies)
 			{
 				// skip zeros
-				while (Start < NumPolicies && !PerDrawingPolicyCounts[Start])
+				while (bCountsAreAccurate && Start < NumPolicies && !PerDrawingPolicyCounts[Start])
 				{
 					Start++;
 				}
@@ -468,7 +502,7 @@ void TStaticMeshDrawList<DrawingPolicyType>::DrawVisibleParallel(
 					while (BatchEnd < NumPolicies - 1 && BatchCount < DrawsPerCmdList)
 					{
 						BatchEnd++;
-						if (PerDrawingPolicyCounts[BatchEnd])
+						if (!bCountsAreAccurate || PerDrawingPolicyCounts[BatchEnd])
 						{
 							BatchCount += PerDrawingPolicyCounts[BatchEnd];
 							LastNonZeroPolicy = BatchEnd;
@@ -484,7 +518,7 @@ void TStaticMeshDrawList<DrawingPolicyType>::DrawVisibleParallel(
 								// this is the last batch and it is small, let merge it
 								UE_CLOG(ParallelCommandListSet.bSpewBalance, LogTemp, Display, TEXT("    Index %d  BatchCount %d    (last merge)"), ParallelCommandListSet.NumParallelCommandLists(), PreviousBatchDraws + BatchCount);
 								FGraphEventRef AnyThreadCompletionEvent = TGraphTask<FDrawVisibleAnyThreadTask<DrawingPolicyType> >::CreateTask(ParallelCommandListSet.GetPrereqs(), ENamedThreads::RenderThread)
-									.ConstructAndDispatchWhenReady(*this, *CmdList, ParallelCommandListSet.View, PolicyContext, StaticMeshVisibilityMap, BatchVisibilityArray, PreviousBatchStart, LastNonZeroPolicy, &PerDrawingPolicyCounts);
+									.ConstructAndDispatchWhenReady(*this, *CmdList, ParallelCommandListSet.View, PolicyContext, StaticMeshVisibilityMap, BatchVisibilityArray, PreviousBatchStart, LastNonZeroPolicy, bCountsAreAccurate ? &PerDrawingPolicyCounts : nullptr);
 								ParallelCommandListSet.AddParallelCommandList(CmdList, AnyThreadCompletionEvent, PreviousBatchDraws + BatchCount);
 								PreviousBatchStart = -1;
 								PreviousBatchEnd = -2;
@@ -495,7 +529,7 @@ void TStaticMeshDrawList<DrawingPolicyType>::DrawVisibleParallel(
 								// this is a decent sized batch, emit the previous batch and save this one for possible merging
 								UE_CLOG(ParallelCommandListSet.bSpewBalance, LogTemp, Display, TEXT("    Index %d  BatchCount %d    "), ParallelCommandListSet.NumParallelCommandLists(), PreviousBatchDraws);
 								FGraphEventRef AnyThreadCompletionEvent = TGraphTask<FDrawVisibleAnyThreadTask<DrawingPolicyType> >::CreateTask(ParallelCommandListSet.GetPrereqs(), ENamedThreads::RenderThread)
-									.ConstructAndDispatchWhenReady(*this, *CmdList, ParallelCommandListSet.View, PolicyContext, StaticMeshVisibilityMap, BatchVisibilityArray, PreviousBatchStart, PreviousBatchEnd, &PerDrawingPolicyCounts);
+									.ConstructAndDispatchWhenReady(*this, *CmdList, ParallelCommandListSet.View, PolicyContext, StaticMeshVisibilityMap, BatchVisibilityArray, PreviousBatchStart, PreviousBatchEnd, bCountsAreAccurate ? &PerDrawingPolicyCounts : nullptr);
 								ParallelCommandListSet.AddParallelCommandList(CmdList, AnyThreadCompletionEvent, PreviousBatchDraws);
 								PreviousBatchStart = Start;
 								PreviousBatchEnd = LastNonZeroPolicy;
@@ -519,31 +553,31 @@ void TStaticMeshDrawList<DrawingPolicyType>::DrawVisibleParallel(
 				UE_CLOG(ParallelCommandListSet.bSpewBalance, LogTemp, Display, TEXT("    Index %d  BatchCount %d    (last)"), ParallelCommandListSet.NumParallelCommandLists(), PreviousBatchDraws);
 				FRHICommandList* CmdList = ParallelCommandListSet.NewParallelCommandList();
 				FGraphEventRef AnyThreadCompletionEvent = TGraphTask<FDrawVisibleAnyThreadTask<DrawingPolicyType> >::CreateTask(ParallelCommandListSet.GetPrereqs(), ENamedThreads::RenderThread)
-					.ConstructAndDispatchWhenReady(*this, *CmdList, ParallelCommandListSet.View, PolicyContext, StaticMeshVisibilityMap, BatchVisibilityArray, PreviousBatchStart, PreviousBatchEnd, &PerDrawingPolicyCounts);
+					.ConstructAndDispatchWhenReady(*this, *CmdList, ParallelCommandListSet.View, PolicyContext, StaticMeshVisibilityMap, BatchVisibilityArray, PreviousBatchStart, PreviousBatchEnd, bCountsAreAccurate ? &PerDrawingPolicyCounts : nullptr);
 				ParallelCommandListSet.AddParallelCommandList(CmdList, AnyThreadCompletionEvent, PreviousBatchDraws);
 			}
 		}
 		else
-	{
-		int32 NumPer = OrderedDrawingPolicies.Num() / EffectiveThreads;
-		int32 Extra = OrderedDrawingPolicies.Num() - NumPer * EffectiveThreads;
-			int32 Start = 0;
-		for (int32 ThreadIndex = 0; ThreadIndex < EffectiveThreads; ThreadIndex++)
 		{
-			int32 Last = Start + (NumPer - 1) + (ThreadIndex < Extra);
-			check(Last >= Start);
-
+			int32 NumPer = OrderedDrawingPolicies.Num() / EffectiveThreads;
+			int32 Extra = OrderedDrawingPolicies.Num() - NumPer * EffectiveThreads;
+				int32 Start = 0;
+			for (int32 ThreadIndex = 0; ThreadIndex < EffectiveThreads; ThreadIndex++)
 			{
-				FRHICommandList* CmdList = ParallelCommandListSet.NewParallelCommandList();
+				int32 Last = Start + (NumPer - 1) + (ThreadIndex < Extra);
+				check(Last >= Start);
 
-					FGraphEventRef AnyThreadCompletionEvent = TGraphTask<FDrawVisibleAnyThreadTask<DrawingPolicyType> >::CreateTask(ParallelCommandListSet.GetPrereqs(), ENamedThreads::RenderThread)
-						.ConstructAndDispatchWhenReady(*this, *CmdList, ParallelCommandListSet.View, PolicyContext, StaticMeshVisibilityMap, BatchVisibilityArray, Start, Last, nullptr);
+				{
+					FRHICommandList* CmdList = ParallelCommandListSet.NewParallelCommandList();
 
-				ParallelCommandListSet.AddParallelCommandList(CmdList, AnyThreadCompletionEvent);
-			}
+						FGraphEventRef AnyThreadCompletionEvent = TGraphTask<FDrawVisibleAnyThreadTask<DrawingPolicyType> >::CreateTask(ParallelCommandListSet.GetPrereqs(), ENamedThreads::RenderThread)
+							.ConstructAndDispatchWhenReady(*this, *CmdList, ParallelCommandListSet.View, PolicyContext, StaticMeshVisibilityMap, BatchVisibilityArray, Start, Last, nullptr);
 
-			Start = Last + 1;
-		}
+					ParallelCommandListSet.AddParallelCommandList(CmdList, AnyThreadCompletionEvent);
+				}
+
+				Start = Last + 1;
+			}
 			check(Start == OrderedDrawingPolicies.Num());
 		}
 	}
diff --git a/Engine/Source/Runtime/ShaderCore/Private/Shader.cpp b/Engine/Source/Runtime/ShaderCore/Private/Shader.cpp
--- a/Engine/Source/Runtime/ShaderCore/Private/Shader.cpp
+++ b/Engine/Source/Runtime/ShaderCore/Private/Shader.cpp
@@ -944,39 +944,52 @@ void FShader::VerifyBoundUniformBufferParameters()
 
 bool FShaderPipelineType::bInitialized = false;
 
-FShaderPipelineType::FShaderPipelineType(const TCHAR* InName, const FShaderType* InVertexShader, const FShaderType* InHullShader, const FShaderType* InDomainShader, const FShaderType* InGeometryShader, const FShaderType* InPixelShader) :
-	VertexShader(InVertexShader),
-	HullShader(InHullShader),
-	DomainShader(InDomainShader),
-	GeometryShader(InGeometryShader),
-	PixelShader(InPixelShader),
+FShaderPipelineType::FShaderPipelineType(
+	const TCHAR* InName,
+	const FShaderType* InVertexShader,
+	const FShaderType* InHullShader,
+	const FShaderType* InDomainShader,
+	const FShaderType* InGeometryShader,
+	const FShaderType* InPixelShader) :
 	Name(InName),
 	TypeName(Name),
 	GlobalListLink(this)
 {
 	checkf(Name && *Name, TEXT("Shader Pipeline Type requires a valid Name!"));
 
-	checkf(VertexShader, TEXT("A Shader Pipeline always requires a Vertex Shader"));
+	checkf(InVertexShader, TEXT("A Shader Pipeline always requires a Vertex Shader"));
 
-	checkf((HullShader == nullptr && DomainShader == nullptr) || (HullShader != nullptr && DomainShader != nullptr), TEXT("Both Hull & Domain shaders are needed for tessellation on Pipeline %s"), Name);
+	checkf((InHullShader == nullptr && InDomainShader == nullptr) || (InHullShader != nullptr && InDomainShader != nullptr), TEXT("Both Hull & Domain shaders are needed for tessellation on Pipeline %s"), Name);
 
 	//make sure the name is shorter than the maximum serializable length
 	check(FCString::Strlen(InName) < NAME_SIZE);
 
-	if (PixelShader)
+	FMemory::Memzero(AllStages);
+
+	if (InPixelShader)
 	{
-		Stages.Add(PixelShader);
+		Stages.Add(InPixelShader);
+		AllStages[SF_Pixel] = InPixelShader;
 	}
-	if (GeometryShader)
+	if (InGeometryShader)
 	{
-		Stages.Add(GeometryShader);
+		Stages.Add(InGeometryShader);
+		AllStages[SF_Geometry] = InGeometryShader;
 	}
-	if (DomainShader)
+	if (InDomainShader)
 	{
-		Stages.Add(DomainShader);
-		Stages.Add(HullShader);
+		Stages.Add(InDomainShader);
+		AllStages[SF_Domain] = InDomainShader;
+
+		Stages.Add(InHullShader);
+		AllStages[SF_Hull] = InHullShader;
 	}
-	Stages.Add(VertexShader);
+	Stages.Add(InVertexShader);
+	AllStages[SF_Vertex] = InVertexShader;
+
+	static uint32 TypeHashCounter = 0;
+	++TypeHashCounter;
+	HashIndex = TypeHashCounter;
 
 	GlobalListLink.LinkHead(GetTypeList());
 	GetNameToTypeMap().Add(FName(InName), this);
@@ -1017,6 +1030,13 @@ void FShaderPipelineType::Initialize()
 	for (TLinkedList<FShaderPipelineType*>::TIterator It(FShaderPipelineType::GetTypeList()); It; It.Next())
 	{
 		const auto* PipelineType = *It;
+
+		// Validate stages
+		for (int32 Index = 0; Index < SF_NumFrequencies; ++Index)
+		{
+			check(!PipelineType->AllStages[Index] || PipelineType->AllStages[Index]->GetFrequency() == (EShaderFrequency)Index);
+		}
+
 		auto& Stages = PipelineType->GetStages();
 
 		// #todo-rco: Do we allow mix/match of global/mesh/material stages?
diff --git a/Engine/Source/Runtime/ShaderCore/Public/Shader.h b/Engine/Source/Runtime/ShaderCore/Public/Shader.h
--- a/Engine/Source/Runtime/ShaderCore/Public/Shader.h
+++ b/Engine/Source/Runtime/ShaderCore/Public/Shader.h
@@ -560,23 +560,23 @@ public:
 	}
 	
 	// Accessors.
-	FShaderType* GetType() const { return Type; }
-	uint32 GetNumInstructions() const { return Resource->NumInstructions; }
-	uint32 GetNumTextureSamplers() const { return Resource->NumTextureSamplers; }
-	const TArray<uint8>& GetCode() const { return Resource->Code; }
-	const FShaderTarget GetTarget() const { return Target; }
+	inline FShaderType* GetType() const { return Type; }
+	inline uint32 GetNumInstructions() const { return Resource->NumInstructions; }
+	inline uint32 GetNumTextureSamplers() const { return Resource->NumTextureSamplers; }
+	inline const TArray<uint8>& GetCode() const { return Resource->Code; }
+	inline const FShaderTarget GetTarget() const { return Target; }
 	inline FSHAHash GetOutputHash() const { return OutputHash; }
 	FShaderId GetId() const;
-	FVertexFactoryType* GetVertexFactoryType() const { return VFType; }
-	FSHAHash GetMaterialShaderMapHash() const { return MaterialShaderMapHash; }
-	int32 GetNumRefs() const { return NumRefs; }
+	inline FVertexFactoryType* GetVertexFactoryType() const { return VFType; }
+	inline FSHAHash GetMaterialShaderMapHash() const { return MaterialShaderMapHash; }
+	inline int32 GetNumRefs() const { return NumRefs; }
 
 	inline FShaderResourceId GetResourceId() const
 	{
 		return Resource->GetId();
 	}
 
-	uint32 GetSizeBytes() const
+	inline uint32 GetSizeBytes() const
 	{
 		return GetTypeSize() + GetAllocatedSize();
 	}
@@ -1001,14 +1001,6 @@ private:
 class SHADERCORE_API FShaderPipelineType
 {
 public:
-	const FShaderType* const VertexShader;
-	const FShaderType* const HullShader;
-	const FShaderType* const DomainShader;
-	const FShaderType* const GeometryShader;
-	const FShaderType* const PixelShader;
-	const TCHAR* const Name;
-	FName TypeName;
-
 	FShaderPipelineType(
 		const TCHAR* InName,
 		const FShaderType* InVertexShader,
@@ -1018,14 +1010,21 @@ public:
 		const FShaderType* InPixelShader);
 	virtual ~FShaderPipelineType();
 
-	bool HasTessellation() const { return HullShader != nullptr; }
-	bool HasGeometry() const { return GeometryShader != nullptr; }
-	bool HasPixelShader() const { return PixelShader != nullptr; }
+	FORCEINLINE bool HasTessellation() const { return AllStages[SF_Domain] != nullptr; }
+	FORCEINLINE bool HasGeometry() const { return AllStages[SF_Geometry] != nullptr; }
+	FORCEINLINE bool HasPixelShader() const { return AllStages[SF_Pixel] != nullptr; }
 
-	FName GetFName() const { return TypeName; }
+	FORCEINLINE const FShaderType* GetShader(EShaderFrequency Frequency) const
+	{
+		check(Frequency < SF_NumFrequencies);
+		return AllStages[Frequency];
+	}
 
-	// Returns an array of valid stages, sorted from PS->GS->DS->HS->VS
-	const TArray<const FShaderType*>& GetStages() const { return Stages; }
+	FORCEINLINE FName GetFName() const { return TypeName; }
+	FORCEINLINE TCHAR const* GetName() const { return Name; }
+
+	// Returns an array of valid stages, sorted from PS->GS->DS->HS->VS, no gaps if missing stages
+	FORCEINLINE const TArray<const FShaderType*>& GetStages() const { return Stages; }
 
 	static TLinkedList<FShaderPipelineType*>*& GetTypeList();
 
@@ -1039,12 +1038,28 @@ public:
 	/** Serializes a shader type reference by name. */
 	SHADERCORE_API friend FArchive& operator<<(FArchive& Ar, const FShaderPipelineType*& Ref);
 
+	/** Hashes a pointer to a shader type. */
+	friend uint32 GetTypeHash(FShaderPipelineType* Ref) { return Ref ? Ref->HashIndex : 0; }
+	friend uint32 GetTypeHash(const FShaderPipelineType* Ref) { return Ref ? Ref->HashIndex : 0; }
+
+	// Check if this pipeline is built of specific types
+	bool IsGlobalTypePipeline() const { return Stages[0]->GetGlobalShaderType() != nullptr; }
+	bool IsMaterialTypePipeline() const { return Stages[0]->GetMaterialShaderType() != nullptr; }
+	bool IsMeshMaterialTypePipeline() const { return Stages[0]->GetMeshMaterialShaderType() != nullptr; }
+
 protected:
-	TLinkedList<FShaderPipelineType*> GlobalListLink;
+	const TCHAR* const Name;
+	FName TypeName;
 
 	// Pipeline Stages, ordered from lowest (usually PS) to highest (VS). Guaranteed at least one stage (for VS).
 	TArray<const FShaderType*> Stages;
 
+	const FShaderType* AllStages[SF_NumFrequencies];
+
+	TLinkedList<FShaderPipelineType*> GlobalListLink;
+
+	uint32 HashIndex;
+
 	static bool bInitialized;
 };
 
@@ -1177,7 +1192,7 @@ public:
 	 * @param bHandleShaderKeyChanges - whether to serialize the data necessary to detect and gracefully handle shader key changes between saving and loading
 	 */
 	void SerializeInline(FArchive& Ar, bool bInlineShaderResource, bool bHandleShaderKeyChanges)
-	{		
+	{
 		if (Ar.IsSaving())
 		{
 			int32 NumShaders = Shaders.Num();
@@ -1224,6 +1239,9 @@ public:
 				Ar << EndOffset;
 				Ar.Seek(EndOffset);
 			}
+
+			int32 NumPipelines = 0;
+			Ar << NumPipelines;
 		}
 
 		if (Ar.IsLoading())
@@ -1277,6 +1295,9 @@ public:
 					Ar.Seek(EndOffset);
 				}
 			}
+
+			int32 NumPipelines = 0;
+			Ar << NumPipelines;
 		}
 	}
 
diff --git a/Engine/Source/Runtime/ShaderCore/Public/ShaderCore.h b/Engine/Source/Runtime/ShaderCore/Public/ShaderCore.h
--- a/Engine/Source/Runtime/ShaderCore/Public/ShaderCore.h
+++ b/Engine/Source/Runtime/ShaderCore/Public/ShaderCore.h
@@ -419,6 +419,9 @@ struct FShaderCompilerInput
 	// e.g. BasePassPixelShader, ReflectionEnvironmentShaders, SlateElementPixelShader, PostProcessCombineLUTs
 	FString SourceFilename;
 	FString EntryPointName;
+	bool bCompilingForShaderPipeline;
+	bool bIncludeUsedOutputs;
+	TArray<FString> UsedOutputs;
 	// Dump debug path (up to platform) e.g. "D:/MMittring-Z3941-A/UE4-Orion/OrionGame/Saved/ShaderDebugInfo/PCD3D_SM5"
 	FString DumpDebugInfoRootPath;
 	// only used if enabled by r.DumpShaderDebugInfo (platform/groupname) e.g. ""
@@ -449,7 +452,6 @@ struct FShaderCompilerInput
 	friend FArchive& operator<<(FArchive& Ar,FShaderCompilerInput& Input)
 	{
 		// Note: this serialize is used to pass between UE4 and the shader compile worker, recompile both when modifying
-
 		Ar << Input.Target;
 		{
 			FString ShaderFormatString(Input.ShaderFormat.ToString());
@@ -459,6 +461,9 @@ struct FShaderCompilerInput
 		Ar << Input.SourceFilePrefix;
 		Ar << Input.SourceFilename;
 		Ar << Input.EntryPointName;
+		Ar << Input.bCompilingForShaderPipeline;
+		Ar << Input.bIncludeUsedOutputs;
+		Ar << Input.UsedOutputs;
 		Ar << Input.DumpDebugInfoRootPath;
 		Ar << Input.DumpDebugInfoPath;
 		Ar << Input.DebugGroupName;
@@ -781,6 +786,8 @@ struct FShaderCompilerOutput
 	:	NumInstructions(0)
 	,	NumTextureSamplers(0)
 	,	bSucceeded(false)
+	,	bFailedRemovingUnused(false)
+	,	bSupportsQueryingUsedAttributes(false)
 	{
 	}
 
@@ -792,14 +799,19 @@ struct FShaderCompilerOutput
 	uint32 NumInstructions;
 	uint32 NumTextureSamplers;
 	bool bSucceeded;
+	bool bFailedRemovingUnused;
+	bool bSupportsQueryingUsedAttributes;
+	TArray<FString> UsedAttributes;
 
 	/** Generates OutputHash from the compiler output. */
 	SHADERCORE_API void GenerateOutputHash();
 	
-	friend FArchive& operator<<(FArchive& Ar,FShaderCompilerOutput& Output)
+	friend FArchive& operator<<(FArchive& Ar, FShaderCompilerOutput& Output)
 	{
 		// Note: this serialize is used to pass between UE4 and the shader compile worker, recompile both when modifying
-		return Ar << Output.ParameterMap << Output.Errors << Output.Target << Output.ShaderCode << Output.NumInstructions << Output.NumTextureSamplers << Output.bSucceeded;
+		Ar << Output.ParameterMap << Output.Errors << Output.Target << Output.ShaderCode << Output.NumInstructions << Output.NumTextureSamplers << Output.bSucceeded;
+		Ar << Output.bFailedRemovingUnused << Output.bSupportsQueryingUsedAttributes << Output.UsedAttributes;
+		return Ar;
 	}
 };
 
diff --git a/Engine/Source/Runtime/Slate/Private/Widgets/Notifications/SNotificationList.cpp b/Engine/Source/Runtime/Slate/Private/Widgets/Notifications/SNotificationList.cpp
--- a/Engine/Source/Runtime/Slate/Private/Widgets/Notifications/SNotificationList.cpp
+++ b/Engine/Source/Runtime/Slate/Private/Widgets/Notifications/SNotificationList.cpp
@@ -268,6 +268,7 @@ class SNotificationItemImpl : public SNotificationExtendable
 public:
 	SLATE_BEGIN_ARGS( SNotificationItemImpl )
 		: _Text()
+		, _Font()
 		, _Image()	
 		, _FadeInDuration(0.5f)
 		, _FadeOutDuration(2.f)
@@ -277,6 +278,8 @@ public:
 
 		/** The text displayed in this text block */
 		SLATE_ATTRIBUTE( FText, Text )
+		/** Sets the font used to draw the text */
+		SLATE_ATTRIBUTE(FSlateFontInfo, Font)
 		/** Setup information for the buttons on the notification */ 
 		SLATE_ATTRIBUTE(TArray<FNotificationButtonInfo>, ButtonDetails)
 		/** The icon image to display next to the text */
@@ -363,9 +366,14 @@ public:
 		];
 
 		{
-			const FSlateFontInfo Font = InArgs._bUseLargeFont.Get() 
-				? FCoreStyle::Get().GetFontStyle(TEXT("NotificationList.FontBold")) 
-				: FCoreStyle::Get().GetFontStyle(TEXT("NotificationList.FontLight"));
+			FSlateFontInfo Font = InArgs._Font.Get();
+
+			if (!Font.HasValidFont())
+			{
+				Font = InArgs._bUseLargeFont.Get()
+					? FCoreStyle::Get().GetFontStyle(TEXT("NotificationList.FontBold"))
+					: FCoreStyle::Get().GetFontStyle(TEXT("NotificationList.FontLight"));
+			}
 
 			// Container for the text and optional interactive widgets (buttons, check box, and hyperlink)
 			TSharedRef<SVerticalBox> TextAndInteractiveWidgetsBox = SNew(SVerticalBox);
@@ -679,6 +687,7 @@ TSharedRef<SNotificationItem> SNotificationList::AddNotification(const FNotifica
 		// Create notification.
 		NewItem = SNew(SNotificationItemImpl)
 			.Text(Info.Text)
+			.Font(Font)
 			.ButtonDetails(Info.ButtonDetails)
 			.Image((Info.Image != nullptr) ? Info.Image : CachedImage)
 			.FadeInDuration(Info.FadeInDuration)
@@ -735,6 +744,7 @@ void SNotificationList::NotificationItemFadedOut (const TSharedRef<SNotification
 void SNotificationList::Construct(const FArguments& InArgs)
 {
 	bDone = false;
+	Font = InArgs._Font;
 
 	ChildSlot
 	[
diff --git a/Engine/Source/Runtime/Slate/Public/Widgets/Notifications/SNotificationList.h b/Engine/Source/Runtime/Slate/Public/Widgets/Notifications/SNotificationList.h
--- a/Engine/Source/Runtime/Slate/Public/Widgets/Notifications/SNotificationList.h
+++ b/Engine/Source/Runtime/Slate/Public/Widgets/Notifications/SNotificationList.h
@@ -211,7 +211,8 @@ class SLATE_API SNotificationList
 public:
 
 	SLATE_BEGIN_ARGS( SNotificationList ){}
-
+		/** Sets the font used to draw the text */
+		SLATE_ATTRIBUTE(FSlateFontInfo, Font)
 	SLATE_END_ARGS()
 	
 	/**
@@ -245,6 +246,9 @@ protected:
 	/** The parent window of this list. */
 	TWeakPtr<SWindow> ParentWindowPtr;
 
+	/** Holds passed in font */
+	TAttribute<FSlateFontInfo> Font;
+
 	/** Flag to auto-destroy this list. */
 	bool bDone;
 };
diff --git a/Engine/Source/Runtime/SlateCore/Private/Rendering/DrawElements.cpp b/Engine/Source/Runtime/SlateCore/Private/Rendering/DrawElements.cpp
--- a/Engine/Source/Runtime/SlateCore/Private/Rendering/DrawElements.cpp
+++ b/Engine/Source/Runtime/SlateCore/Private/Rendering/DrawElements.cpp
@@ -193,6 +193,15 @@ void FSlateDrawElement::MakeLines(FSlateWindowElementList& ElementList, uint32 I
 void FSlateDrawElement::MakeViewport( FSlateWindowElementList& ElementList, uint32 InLayer, const FPaintGeometry& PaintGeometry, TSharedPtr<const ISlateViewport> Viewport, const FSlateRect& InClippingRect, bool bGammaCorrect, bool bAllowBlending, ESlateDrawEffect::Type InDrawEffects, const FLinearColor& InTint )
 {
 	SCOPE_CYCLE_COUNTER( STAT_SlateDrawElementMakeTime )
+
+	// If we don't allow blending, the shader needs to know it needs to treat the incoming
+	// source data as already pre-multiplied, so it doesn't attempt to multiply through using
+	// the alpha.
+	if ( !bAllowBlending )
+	{
+		InDrawEffects |= ESlateDrawEffect::PreMultipliedAlpha;
+	}
+
 	PaintGeometry.CommitTransformsIfUsingLegacyConstructor();
 	FSlateDrawElement& DrawElt = ElementList.AddUninitialized();
 	DrawElt.Init(InLayer, PaintGeometry, InClippingRect, InDrawEffects);
diff --git a/Engine/Source/Runtime/SlateCore/Public/Input/Events.h b/Engine/Source/Runtime/SlateCore/Public/Input/Events.h
--- a/Engine/Source/Runtime/SlateCore/Public/Input/Events.h
+++ b/Engine/Source/Runtime/SlateCore/Public/Input/Events.h
@@ -114,9 +114,15 @@ private:
 struct FVirtualPointerPosition
 {
 	FVirtualPointerPosition()
-	: CurrentCursorPosition(FVector2D::ZeroVector)
-	, LastCursorPosition(FVector2D::ZeroVector)
+		: CurrentCursorPosition(FVector2D::ZeroVector)
+		, LastCursorPosition(FVector2D::ZeroVector)
 	{}
+
+	FVirtualPointerPosition(const FVector2D& InCurrentCursorPosition, const FVector2D& InLastCursorPosition)
+		: CurrentCursorPosition(InCurrentCursorPosition)
+		, LastCursorPosition(InLastCursorPosition)
+	{}
+
 	FVector2D CurrentCursorPosition;
 	FVector2D LastCursorPosition;
 };
diff --git a/Engine/Source/Runtime/SlateCore/Public/Rendering/RenderingCommon.h b/Engine/Source/Runtime/SlateCore/Public/Rendering/RenderingCommon.h
--- a/Engine/Source/Runtime/SlateCore/Public/Rendering/RenderingCommon.h
+++ b/Engine/Source/Runtime/SlateCore/Public/Rendering/RenderingCommon.h
@@ -8,7 +8,7 @@ struct FVector2D;
 class FSlateRect;
 enum class EPopupMethod : uint8;
 
-
+#define SLATE_PRE_MULTIPLY 1
 #define SLATE_USE_32BIT_INDICES !PLATFORM_USES_ES2
 
 #if SLATE_USE_32BIT_INDICES
@@ -58,9 +58,11 @@ namespace ESlateDrawEffect
 	/** No effect applied */
 	const Type None = 0;
 	/** Draw the element with a disabled effect */
-	const Type DisabledEffect = 1<<0;
+	const Type DisabledEffect = 1 << 0;
 	/** Don't read from texture alpha channel */
-	const Type IgnoreTextureAlpha = 1<<2;
+	const Type IgnoreTextureAlpha = 1 << 1;
+	/** If this is provided, we will not perform pre-multiply of alpha of the source texture in the slate pixel shader, expecting it to already be the case. */
+	const Type PreMultipliedAlpha = 1 << 2;
 };
 
 
diff --git a/Engine/Source/Runtime/SlateRHIRenderer/Private/Slate3DRenderer.cpp b/Engine/Source/Runtime/SlateRHIRenderer/Private/Slate3DRenderer.cpp
--- a/Engine/Source/Runtime/SlateRHIRenderer/Private/Slate3DRenderer.cpp
+++ b/Engine/Source/Runtime/SlateRHIRenderer/Private/Slate3DRenderer.cpp
@@ -4,13 +4,13 @@
 #include "Slate3DRenderer.h"
 #include "ElementBatcher.h"
 
-FSlate3DRenderer::FSlate3DRenderer( TSharedPtr<FSlateRHIResourceManager> InResourceManager, TSharedPtr<FSlateFontCache> InFontCache )
+FSlate3DRenderer::FSlate3DRenderer( TSharedPtr<FSlateRHIResourceManager> InResourceManager, TSharedPtr<FSlateFontCache> InFontCache, bool bUseGammaCorrection )
 	: ResourceManager( InResourceManager.ToSharedRef() )
 	, FontCache( InFontCache.ToSharedRef() )
 {
 
 	RenderTargetPolicy = MakeShareable( new FSlateRHIRenderingPolicy( FontCache, ResourceManager ) );
-	RenderTargetPolicy->SetUseGammaCorrection( false );
+	RenderTargetPolicy->SetUseGammaCorrection( bUseGammaCorrection );
 
 	ElementBatcher = MakeShareable(new FSlateElementBatcher(RenderTargetPolicy.ToSharedRef()));
 }
@@ -83,8 +83,14 @@ void FSlate3DRenderer::DrawWindow_GameThread(FSlateDrawBuffer& DrawBuffer)
 
 void FSlate3DRenderer::DrawWindowToTarget_RenderThread( FRHICommandListImmediate& RHICmdList, UTextureRenderTarget2D* RenderTarget, FSlateDrawBuffer& InDrawBuffer )
 {
-	FSlateBatchData& BatchData = InDrawBuffer.GetWindowElementLists()[0]->GetBatchData();
-	FElementBatchMap& RootBatchMap = InDrawBuffer.GetWindowElementLists()[0]->GetRootDrawLayer().GetElementBatchMap();
+	checkSlow(InDrawBuffer.GetWindowElementLists().Num() == 1);
+
+	FSlateWindowElementList& WindowElementList = *InDrawBuffer.GetWindowElementLists()[0].Get();
+
+	FSlateBatchData& BatchData = WindowElementList.GetBatchData();
+	FElementBatchMap& RootBatchMap = WindowElementList.GetRootDrawLayer().GetElementBatchMap();
+
+	WindowElementList.PreDraw_ParallelThread();
 
 	BatchData.CreateRenderBatches(RootBatchMap);
 	RenderTargetPolicy->BeginDrawingWindows();
@@ -120,9 +126,16 @@ void FSlate3DRenderer::DrawWindowToTarget_RenderThread( FRHICommandListImmediate
 		);
 	}
 
+	for ( TSharedPtr<FSlateWindowElementList>& ElementList : InDrawBuffer.GetWindowElementLists() )
+	{
+		ElementList->PostDraw_ParallelThread();
+	}
+
 	InDrawBuffer.Unlock();
 
 	RenderTargetPolicy->EndDrawingWindows();
 
 	RHICmdList.CopyToResolveTarget(RenderTargetResource->GetTextureRHI(), RTResource, true, FResolveParams());
+	
+
 }
\ No newline at end of file
diff --git a/Engine/Source/Runtime/SlateRHIRenderer/Private/Slate3DRenderer.h b/Engine/Source/Runtime/SlateRHIRenderer/Private/Slate3DRenderer.h
--- a/Engine/Source/Runtime/SlateRHIRenderer/Private/Slate3DRenderer.h
+++ b/Engine/Source/Runtime/SlateRHIRenderer/Private/Slate3DRenderer.h
@@ -7,8 +7,8 @@
 class FSlate3DRenderer : public ISlate3DRenderer
 {
 public:
-	FSlate3DRenderer( TSharedPtr<FSlateRHIResourceManager> InResourceManager, TSharedPtr<FSlateFontCache> InFontCache );
-	~FSlate3DRenderer();
+	FSlate3DRenderer( TSharedPtr<FSlateRHIResourceManager> InResourceManager, TSharedPtr<FSlateFontCache> InFontCache, bool bUseGammaCorrection = false );
+	virtual ~FSlate3DRenderer();
 
 	virtual FSlateDrawBuffer& GetDrawBuffer() override;
 	virtual void DrawWindow_GameThread(FSlateDrawBuffer& DrawBuffer) override;
@@ -32,5 +32,4 @@ private:
 
 	/** The draw buffer that is currently free for use by the game thread */
 	uint8 FreeBufferIndex;
-
-};
\ No newline at end of file
+};
diff --git a/Engine/Source/Runtime/SlateRHIRenderer/Private/SlateElementIndexBuffer.cpp b/Engine/Source/Runtime/SlateRHIRenderer/Private/SlateElementIndexBuffer.cpp
new file mode 100644
--- /dev/null
+++ b/Engine/Source/Runtime/SlateRHIRenderer/Private/SlateElementIndexBuffer.cpp
@@ -0,0 +1,126 @@
+// Copyright 1998-2015 Epic Games, Inc. All Rights Reserved.
+
+#include "SlateRHIRendererPrivatePCH.h"
+#include "SlateElementIndexBuffer.h"
+
+FSlateElementIndexBuffer::FSlateElementIndexBuffer()
+	: BufferSize(0)	 
+	, MinBufferSize(0)
+	, BufferUsageSize(0)
+{
+
+}
+
+FSlateElementIndexBuffer::~FSlateElementIndexBuffer()
+{
+
+}
+
+void FSlateElementIndexBuffer::Init( int32 MinNumIndices )
+{
+	MinBufferSize = sizeof(SlateIndex) * FMath::Max( MinNumIndices, 200 );
+
+	if ( IsInRenderingThread() )
+	{
+		InitResource();
+	}
+	else
+	{
+		BeginInitResource(this);
+	}
+}
+
+void FSlateElementIndexBuffer::Destroy()
+{
+	if ( IsInRenderingThread() )
+	{
+		ReleaseResource();
+	}
+	else
+	{
+		BeginReleaseResource(this);
+	}
+}
+
+
+/** Initializes the index buffers RHI resource. */
+void FSlateElementIndexBuffer::InitDynamicRHI()
+{
+	checkSlow( IsInRenderingThread() );
+
+	check( MinBufferSize > 0 );
+
+	BufferSize = MinBufferSize;
+
+	FRHIResourceCreateInfo CreateInfo;
+	IndexBufferRHI = RHICreateIndexBuffer( sizeof(SlateIndex), MinBufferSize, BUF_Dynamic, CreateInfo );
+	check( IsValidRef(IndexBufferRHI) );
+}
+
+/** Resizes the buffer to the passed in size.  Preserves internal data */
+void FSlateElementIndexBuffer::ResizeBuffer( int32 NewSizeBytes )
+{
+	checkSlow( IsInRenderingThread() );
+
+	int32 FinalSize = FMath::Max( NewSizeBytes, MinBufferSize );
+
+	if( FinalSize != 0 && FinalSize != BufferSize )
+	{
+		IndexBufferRHI.SafeRelease();
+		FRHIResourceCreateInfo CreateInfo;
+		IndexBufferRHI = RHICreateIndexBuffer( sizeof(SlateIndex), FinalSize, BUF_Dynamic, CreateInfo );
+		check(IsValidRef(IndexBufferRHI));
+
+		BufferSize = FinalSize;
+	}
+}
+
+void FSlateElementIndexBuffer::PreFillBuffer(int32 RequiredIndexCount, bool bShrinkToMinSize)
+{
+	//SCOPE_CYCLE_COUNTER( STAT_SlatePreFullBufferRTTime );
+
+	checkSlow(IsInRenderingThread());
+
+	if (RequiredIndexCount > 0)
+	{
+		int32 RequiredBufferSize = RequiredIndexCount*sizeof(SlateIndex);
+
+		// resize if needed
+		if (RequiredBufferSize > GetBufferSize() || bShrinkToMinSize)
+		{
+			// Use array resize techniques for the vertex buffer
+			ResizeBuffer(RequiredBufferSize);
+		}
+
+		BufferUsageSize = RequiredBufferSize;		
+	}
+}
+
+void* FSlateElementIndexBuffer::LockBuffer_RenderThread(int32 NumIndices)
+{
+	uint32 RequiredBufferSize = NumIndices*sizeof(SlateIndex);		
+	return RHILockIndexBuffer( IndexBufferRHI, 0, RequiredBufferSize, RLM_WriteOnly );
+}
+
+void FSlateElementIndexBuffer::UnlockBuffer_RenderThread()
+{
+	RHIUnlockIndexBuffer( IndexBufferRHI );
+}
+
+void* FSlateElementIndexBuffer::LockBuffer_RHIThread(int32 NumIndices)
+{
+	uint32 RequiredBufferSize = NumIndices*sizeof(SlateIndex);		
+	return GDynamicRHI->RHILockIndexBuffer( IndexBufferRHI, 0, RequiredBufferSize, RLM_WriteOnly );
+}
+
+void FSlateElementIndexBuffer::UnlockBuffer_RHIThread()
+{
+	GDynamicRHI->RHIUnlockIndexBuffer( IndexBufferRHI );
+}
+
+/** Releases the index buffers RHI resource. */
+void FSlateElementIndexBuffer::ReleaseDynamicRHI()
+{
+	IndexBufferRHI.SafeRelease();
+	BufferSize = 0;
+}
\ No newline at end of file
diff --git a/Engine/Source/Runtime/SlateRHIRenderer/Private/SlateElementIndexBuffer.h b/Engine/Source/Runtime/SlateRHIRenderer/Private/SlateElementIndexBuffer.h
new file mode 100644
--- /dev/null
+++ b/Engine/Source/Runtime/SlateRHIRenderer/Private/SlateElementIndexBuffer.h
@@ -0,0 +1,65 @@
+// Copyright 1998-2015 Epic Games, Inc. All Rights Reserved.
+
+#pragma once
+
+class FSlateElementIndexBuffer : public FIndexBuffer
+{
+public:
+	/**
+	 * Constructor
+	 * 
+	 * @param MinNumIndices	The minimum number of indices this buffer should always support
+	 */
+	FSlateElementIndexBuffer();
+	~FSlateElementIndexBuffer();
+
+	void Init( int32 MinNumIndices );
+	void Destroy();
+
+	/** Initializes the index buffers RHI resource. */
+	virtual void InitDynamicRHI() override;
+
+	/** Releases the index buffers RHI resource. */
+	virtual void ReleaseDynamicRHI() override;
+
+	/** Returns a friendly name for this buffer. */
+	virtual FString GetFriendlyName() const override { return TEXT("SlateElementIndices"); }
+
+	/** Returns the size of this buffer */
+	int32 GetBufferSize() const { return BufferSize; }
+
+	/** Returns the used size of this buffer */
+	int32 GetBufferUsageSize() const { return BufferUsageSize; }
+
+	/** Resets the usage of the buffer */
+	void ResetBufferUsage() { BufferUsageSize = 0; }
+
+	/** Resizes buffer, accumulates states safely on render thread */
+	void PreFillBuffer(int32 RequiredIndexCount, bool bShrinkToMinSize);
+
+	void* LockBuffer_RenderThread(int32 NumIndices);
+	void UnlockBuffer_RenderThread();
+
+	void* LockBuffer_RHIThread(int32 NumIndices);
+	void UnlockBuffer_RHIThread();
+
+	int32 GetMinBufferSize() const { return MinBufferSize; }
+
+private:
+	/** Resizes the buffer to the passed in size.  Preserves internal data */
+	void ResizeBuffer( int32 NewSizeBytes );
+private:
+	/** The current size of the buffer in bytes. */
+	int32 BufferSize;
+	
+	/** The minimum size the buffer should always be */
+	int32 MinBufferSize;
+
+	/** The size of the currently used portion of the buffer */
+	int32 BufferUsageSize;
+
+	/** Hidden copy methods. */
+	FSlateElementIndexBuffer( const FSlateElementIndexBuffer& );
+	void operator=(const FSlateElementIndexBuffer& );
+
+};
\ No newline at end of file
diff --git a/Engine/Source/Runtime/SlateRHIRenderer/Private/SlateElementVertexBuffer.h b/Engine/Source/Runtime/SlateRHIRenderer/Private/SlateElementVertexBuffer.h
new file mode 100644
--- /dev/null
+++ b/Engine/Source/Runtime/SlateRHIRenderer/Private/SlateElementVertexBuffer.h
@@ -0,0 +1,167 @@
+// Copyright 1998-2015 Epic Games, Inc. All Rights Reserved.
+
+#pragma once
+
+/**
+ * Vertex buffer containing all Slate vertices
+ */
+template <typename VertexType>
+class TSlateElementVertexBuffer : public FVertexBuffer
+{
+public:
+	TSlateElementVertexBuffer()
+		: BufferSize(0)
+		, MinBufferSize(0)
+		, BufferUsageSize(0)
+	{}
+
+	~TSlateElementVertexBuffer() {};
+
+	void Init( int32 MinNumVertices )
+	{
+		MinBufferSize = sizeof(VertexType) * FMath::Max( MinNumVertices, 200 );
+
+		if ( IsInRenderingThread() )
+		{
+			InitResource();
+		}
+		else
+		{
+			BeginInitResource(this);
+		}
+	}
+
+	void Destroy()
+	{
+		if ( IsInRenderingThread() )
+		{
+			ReleaseResource();
+		}
+		else
+		{
+			BeginReleaseResource(this);
+		}
+	}
+
+	/** Initializes the vertex buffers RHI resource. */
+	virtual void InitDynamicRHI()
+	{
+		if( !IsValidRef(VertexBufferRHI) )
+		{
+			check( MinBufferSize > 0 )
+	
+			BufferSize = MinBufferSize;
+
+			FRHIResourceCreateInfo CreateInfo;
+			VertexBufferRHI = RHICreateVertexBuffer( MinBufferSize, BUF_Dynamic, CreateInfo );
+
+			// Ensure the vertex buffer could be created
+			check(IsValidRef(VertexBufferRHI));
+		}
+	}
+
+	/** Releases the vertex buffers RHI resource. */
+	virtual void ReleaseDynamicRHI()
+	{
+		VertexBufferRHI.SafeRelease();
+		BufferSize = 0;
+	}
+
+	/** Returns a friendly name for this buffer. */
+	virtual FString GetFriendlyName() const { return TEXT("SlateElementVertices"); }
+
+	/** Returns the size of the buffer in bytes. */
+	int32 GetBufferSize() const { return BufferSize; }
+
+	/** Returns the used size of this buffer */
+	int32 GetBufferUsageSize() const { return BufferUsageSize; }
+
+	/** Resets the usage of the buffer */
+	void ResetBufferUsage() { BufferUsageSize = 0; }
+
+	/** Resizes buffer, accumulates states safely on render thread */
+	void PreFillBuffer(int32 RequiredVertexCount, bool bShrinkToMinSize)
+	{
+		checkSlow(IsInRenderingThread());
+
+		if (RequiredVertexCount > 0 )
+		{
+#if !SLATE_USE_32BIT_INDICES
+			// make sure our index buffer can handle this
+			checkf(RequiredVertexCount < 0xFFFF, TEXT("Slate vertex buffer is too large (%d) to work with uint16 indices"), RequiredVertexCount);
+#endif
+			int32 RequiredBufferSize = RequiredVertexCount*sizeof(VertexType);
+
+			// resize if needed
+			if(RequiredBufferSize > GetBufferSize() || bShrinkToMinSize)
+			{
+				ResizeBuffer(RequiredBufferSize);
+			}
+
+			BufferUsageSize = RequiredBufferSize;
+		}
+
+	}
+
+	int32 GetMinBufferSize() const { return MinBufferSize; }
+
+	void* LockBuffer_RenderThread(int32 NumVertices)
+	{
+		uint32 RequiredBufferSize = NumVertices*sizeof(VertexType);	
+
+		return RHILockVertexBuffer( VertexBufferRHI, 0, RequiredBufferSize, RLM_WriteOnly );
+	}
+
+	void UnlockBuffer_RenderThread()
+	{
+		RHIUnlockVertexBuffer( VertexBufferRHI );
+	}
+
+	void* LockBuffer_RHIThread(int32 NumVertices)
+	{
+		int32 RequiredBufferSize = NumVertices*sizeof(VertexType);
+
+		//use non-flushing DynamicRHI version when on RHIThread.
+		return GDynamicRHI->RHILockVertexBuffer(VertexBufferRHI, 0, RequiredBufferSize, RLM_WriteOnly);
+	}
+
+	void UnlockBuffer_RHIThread()
+	{
+		GDynamicRHI->RHIUnlockVertexBuffer(VertexBufferRHI);		
+	}
+
+private:
+	/** Resizes the buffer to the passed in size.  Preserves internal data*/
+	void ResizeBuffer( int32 NewSizeBytes )
+	{
+		checkSlow( IsInRenderingThread() );
+
+		int32 FinalSize = FMath::Max( NewSizeBytes, MinBufferSize );
+
+		if( FinalSize != 0 && FinalSize != BufferSize )
+		{
+			VertexBufferRHI.SafeRelease();
+
+			FRHIResourceCreateInfo CreateInfo;
+			VertexBufferRHI = RHICreateVertexBuffer(FinalSize, BUF_Dynamic, CreateInfo);
+
+			check(IsValidRef(VertexBufferRHI));
+
+			BufferSize = FinalSize;
+		}
+	}
+
+private:
+	/** The size of the buffer in bytes. */
+	int32 BufferSize;
+	
+	/** The minimum size the buffer should always be */
+	int32 MinBufferSize;
+
+	/** The size of the used portion of the buffer */
+	int32 BufferUsageSize;
+
+	/** Hidden copy methods. */
+	TSlateElementVertexBuffer( const TSlateElementVertexBuffer& );
+	void operator=(const TSlateElementVertexBuffer& );
+};
diff --git a/Engine/Source/Runtime/SlateRHIRenderer/Private/SlateMaterialShader.cpp b/Engine/Source/Runtime/SlateRHIRenderer/Private/SlateMaterialShader.cpp
--- a/Engine/Source/Runtime/SlateRHIRenderer/Private/SlateMaterialShader.cpp
+++ b/Engine/Source/Runtime/SlateRHIRenderer/Private/SlateMaterialShader.cpp
@@ -69,6 +69,7 @@ void FSlateMaterialShaderPS::ModifyCompilationEnvironment(EShaderPlatform Platfo
 	// Set defines based on what this shader will be used for
 	OutEnvironment.SetDefine( TEXT("USE_MATERIALS"), 1 );
 	OutEnvironment.SetDefine( TEXT("NUM_CUSTOMIZED_UVS"), Material->GetNumCustomizedUVs() );
+	OutEnvironment.SetDefine( TEXT("PRE_MULTIPLIED_ALPHA"), (uint32)0 );
 
 	FMaterialShader::ModifyCompilationEnvironment( Platform, Material, OutEnvironment );
 }
diff --git a/Engine/Source/Runtime/SlateRHIRenderer/Private/SlateRHIRendererModule.cpp b/Engine/Source/Runtime/SlateRHIRenderer/Private/SlateRHIRendererModule.cpp
--- a/Engine/Source/Runtime/SlateRHIRenderer/Private/SlateRHIRendererModule.cpp
+++ b/Engine/Source/Runtime/SlateRHIRenderer/Private/SlateRHIRendererModule.cpp
@@ -60,11 +60,11 @@ public:
 		return MakeShareable( new FSlateRHIRenderer( ResourceManager, FontCache, FontMeasure ) );
 	}
 
-	virtual TSharedRef<ISlate3DRenderer> CreateSlate3DRenderer() override
+	virtual TSharedRef<ISlate3DRenderer> CreateSlate3DRenderer(bool bUseGammaCorrection) override
 	{
 		ConditionalCreateResources();
 
-		return MakeShareable( new FSlate3DRenderer( ResourceManager, FontCache ) );
+		return MakeShareable( new FSlate3DRenderer( ResourceManager, FontCache, bUseGammaCorrection ) );
 	}
 
 	virtual TSharedRef<ISlateFontAtlasFactory> CreateSlateFontAtlasFactory() override
diff --git a/Engine/Source/Runtime/SlateRHIRenderer/Private/SlateRHIRenderingPolicy.cpp b/Engine/Source/Runtime/SlateRHIRenderer/Private/SlateRHIRenderingPolicy.cpp
--- a/Engine/Source/Runtime/SlateRHIRenderer/Private/SlateRHIRenderingPolicy.cpp
+++ b/Engine/Source/Runtime/SlateRHIRenderer/Private/SlateRHIRenderingPolicy.cpp
@@ -12,6 +12,8 @@
 #include "SlateMaterialResource.h"
 #include "Rendering/DrawElements.h"
 #include "SlateUpdatableBuffer.h"
+#include "SlateElementIndexBuffer.h"
+#include "SlateElementVertexBuffer.h"
 
 DECLARE_CYCLE_STAT(TEXT("Update Buffers RT"), STAT_SlateUpdateBufferRTTime, STATGROUP_Slate);
 DECLARE_CYCLE_STAT(TEXT("PreFill Buffers RT"), STAT_SlatePreFullBufferRTTime, STATGROUP_Slate);
@@ -22,134 +24,12 @@ DECLARE_DWORD_COUNTER_STAT(TEXT("Num Vertices"), STAT_SlateVertexCount, STATGROU
 DECLARE_MEMORY_STAT(TEXT("Vertex Buffer Memory"), STAT_SlateVertexBufferMemory, STATGROUP_SlateMemory);
 DECLARE_MEMORY_STAT(TEXT("Index Buffer Memory"), STAT_SlateIndexBufferMemory, STATGROUP_SlateMemory);
 
-FSlateElementIndexBuffer::FSlateElementIndexBuffer()
-	: BufferSize(0)	 
-	, MinBufferSize(0)
-	, BufferUsageSize(0)
-{
-
-}
-
-FSlateElementIndexBuffer::~FSlateElementIndexBuffer()
-{
-
-}
-
-void FSlateElementIndexBuffer::Init( int32 MinNumIndices )
-{
-	MinBufferSize = sizeof(SlateIndex) * FMath::Max( MinNumIndices, 200 );
-
-	if ( IsInRenderingThread() )
-	{
-		InitResource();
-	}
-	else
-	{
-		BeginInitResource(this);
-	}
-}
-
-void FSlateElementIndexBuffer::Destroy()
-{
-	if ( IsInRenderingThread() )
-	{
-		ReleaseResource();
-	}
-	else
-	{
-		BeginReleaseResource(this);
-	}
-}
-
-
-/** Initializes the index buffers RHI resource. */
-void FSlateElementIndexBuffer::InitDynamicRHI()
-{
-	checkSlow( IsInRenderingThread() );
-
-	check( MinBufferSize > 0 );
-
-	BufferSize = MinBufferSize;
-
-	FRHIResourceCreateInfo CreateInfo;
-	IndexBufferRHI = RHICreateIndexBuffer( sizeof(SlateIndex), MinBufferSize, BUF_Dynamic, CreateInfo );
-	check( IsValidRef(IndexBufferRHI) );
-}
-
-/** Resizes the buffer to the passed in size.  Preserves internal data */
-void FSlateElementIndexBuffer::ResizeBuffer( int32 NewSizeBytes )
-{
-	checkSlow( IsInRenderingThread() );
-
-	int32 FinalSize = FMath::Max( NewSizeBytes, MinBufferSize );
-
-	if( FinalSize != 0 && FinalSize != BufferSize )
-	{
-		IndexBufferRHI.SafeRelease();
-		FRHIResourceCreateInfo CreateInfo;
-		IndexBufferRHI = RHICreateIndexBuffer( sizeof(SlateIndex), FinalSize, BUF_Dynamic, CreateInfo );
-		check(IsValidRef(IndexBufferRHI));
-
-		BufferSize = FinalSize;
-	}
-}
-
-void FSlateElementIndexBuffer::PreFillBuffer(int32 RequiredIndexCount, bool bShrinkToMinSize)
-{
-	//SCOPE_CYCLE_COUNTER( STAT_SlatePreFullBufferRTTime );
-
-	checkSlow(IsInRenderingThread());
-
-	if (RequiredIndexCount > 0)
-	{
-		int32 RequiredBufferSize = RequiredIndexCount*sizeof(SlateIndex);
-
-		// resize if needed
-		if (RequiredBufferSize > GetBufferSize() || bShrinkToMinSize)
-		{
-			// Use array resize techniques for the vertex buffer
-			ResizeBuffer(RequiredBufferSize);
-		}
-
-		BufferUsageSize = RequiredBufferSize;		
-	}
-}
-
-void* FSlateElementIndexBuffer::LockBuffer_RenderThread(int32 NumIndices)
-{
-	uint32 RequiredBufferSize = NumIndices*sizeof(SlateIndex);		
-	return RHILockIndexBuffer( IndexBufferRHI, 0, RequiredBufferSize, RLM_WriteOnly );
-}
-
-void FSlateElementIndexBuffer::UnlockBuffer_RenderThread()
-{
-	RHIUnlockIndexBuffer( IndexBufferRHI );
-}
-
-void* FSlateElementIndexBuffer::LockBuffer_RHIThread(int32 NumIndices)
-{
-	uint32 RequiredBufferSize = NumIndices*sizeof(SlateIndex);		
-	return GDynamicRHI->RHILockIndexBuffer( IndexBufferRHI, 0, RequiredBufferSize, RLM_WriteOnly );
-}
-
-void FSlateElementIndexBuffer::UnlockBuffer_RHIThread()
-{
-	GDynamicRHI->RHIUnlockIndexBuffer( IndexBufferRHI );
-}
-
-/** Releases the index buffers RHI resource. */
-void FSlateElementIndexBuffer::ReleaseDynamicRHI()
-{
-	IndexBufferRHI.SafeRelease();
-	BufferSize = 0;
-}
-
 FSlateRHIRenderingPolicy::FSlateRHIRenderingPolicy( TSharedPtr<FSlateFontCache> InFontCache, TSharedRef<FSlateRHIResourceManager> InResourceManager )
 	: FSlateRenderingPolicy(0)
 	, ResourceManager( InResourceManager )
 	, FontCache( InFontCache )
-	, CurrentBufferIndex(0)
 	, bGammaCorrect(true)
+	, CurrentBufferIndex(0)
 {
 	InitResources();
 };
@@ -192,49 +72,11 @@ void FSlateRHIRenderingPolicy::ReleaseResources()
 		VertexBuffers[BufferIndex].Destroy();
 		IndexBuffers[BufferIndex].Destroy();
 	}
-
-	for ( TCachedBufferMap::TIterator BufferIt(CachedBuffers); BufferIt; ++BufferIt )
-	{
-		FSlateRenderDataHandle* Handle = BufferIt.Key();
-		FCachedRenderBuffers* Buffer = BufferIt.Value();
-
-		Handle->Disconnect();
-
-		Buffer->VertexBuffer.Destroy();
-		Buffer->IndexBuffer.Destroy();
-	}
-
-	for ( TCachedBufferPoolMap::TIterator BufferIt(CachedBufferPool); BufferIt; ++BufferIt )
-	{
-		TArray< FCachedRenderBuffers* >& Pool = BufferIt.Value();
-		for ( FCachedRenderBuffers* PooledBuffer : Pool )
-		{
-			PooledBuffer->VertexBuffer.Destroy();
-			PooledBuffer->IndexBuffer.Destroy();
-		}
-	}
 }
 
 void FSlateRHIRenderingPolicy::DeleteReleasedResources()
 {
-	for ( TCachedBufferMap::TIterator BufferIt(CachedBuffers); BufferIt; ++BufferIt )
-	{
-		FCachedRenderBuffers* Buffer = BufferIt.Value();
-		delete Buffer;
-	}
-
-	CachedBuffers.Empty();
-
-	for ( TCachedBufferPoolMap::TIterator BufferIt(CachedBufferPool); BufferIt; ++BufferIt )
-	{
-		TArray< FCachedRenderBuffers* >& Pool = BufferIt.Value();
-		for ( FCachedRenderBuffers* PooledBuffer : Pool )
-		{
-			delete PooledBuffer;
-		}
-	}
-
-	CachedBufferPool.Empty();
+	
 }
 
 void FSlateRHIRenderingPolicy::BeginDrawingWindows()
@@ -265,29 +107,29 @@ void FSlateRHIRenderingPolicy::EndDrawingWindows()
 		TotalIndexBufferUsage += IndexBuffers[BufferIndex].GetBufferUsageSize();
 	}
 
-	for ( TCachedBufferMap::TIterator BufferIt(CachedBuffers); BufferIt; ++BufferIt )
-	{
-		FCachedRenderBuffers* PooledBuffer = BufferIt.Value();
+	//for ( TCachedBufferMap::TIterator BufferIt(CachedBuffers); BufferIt; ++BufferIt )
+	//{
+	//	FCachedRenderBuffers* PooledBuffer = BufferIt.Value();
 
-		TotalVertexBufferMemory += PooledBuffer->VertexBuffer.GetBufferSize();
-		TotalVertexBufferUsage += PooledBuffer->VertexBuffer.GetBufferUsageSize();
+	//	TotalVertexBufferMemory += PooledBuffer->VertexBuffer.GetBufferSize();
+	//	TotalVertexBufferUsage += PooledBuffer->VertexBuffer.GetBufferUsageSize();
 
-		TotalIndexBufferMemory += PooledBuffer->IndexBuffer.GetBufferSize();
-		TotalIndexBufferUsage += PooledBuffer->IndexBuffer.GetBufferUsageSize();
-	}
+	//	TotalIndexBufferMemory += PooledBuffer->IndexBuffer.GetBufferSize();
+	//	TotalIndexBufferUsage += PooledBuffer->IndexBuffer.GetBufferUsageSize();
+	//}
 
-	for ( TCachedBufferPoolMap::TIterator BufferIt(CachedBufferPool); BufferIt; ++BufferIt )
-	{
-		TArray< FCachedRenderBuffers* >& Pool = BufferIt.Value();
-		for ( FCachedRenderBuffers* PooledBuffer : Pool )
-		{
-			TotalVertexBufferMemory += PooledBuffer->VertexBuffer.GetBufferSize();
-			TotalVertexBufferUsage += PooledBuffer->VertexBuffer.GetBufferUsageSize();
+	//for ( TCachedBufferPoolMap::TIterator BufferIt(CachedBufferPool); BufferIt; ++BufferIt )
+	//{
+	//	TArray< FCachedRenderBuffers* >& Pool = BufferIt.Value();
+	//	for ( FCachedRenderBuffers* PooledBuffer : Pool )
+	//	{
+	//		TotalVertexBufferMemory += PooledBuffer->VertexBuffer.GetBufferSize();
+	//		TotalVertexBufferUsage += PooledBuffer->VertexBuffer.GetBufferUsageSize();
 
-			TotalIndexBufferMemory += PooledBuffer->IndexBuffer.GetBufferSize();
-			TotalIndexBufferUsage += PooledBuffer->IndexBuffer.GetBufferUsageSize();
-		}
-	}
+	//		TotalIndexBufferMemory += PooledBuffer->IndexBuffer.GetBufferSize();
+	//		TotalIndexBufferUsage += PooledBuffer->IndexBuffer.GetBufferUsageSize();
+	//	}
+	//}
 
 	SET_MEMORY_STAT( STAT_SlateVertexBufferMemory, TotalVertexBufferMemory );
 	SET_MEMORY_STAT( STAT_SlateIndexBufferMemory, TotalIndexBufferMemory );
@@ -333,77 +175,19 @@ void FSlateRHIRenderingPolicy::UpdateVertexAndIndexBuffers(FRHICommandListImmedi
 	// Should only be called by the rendering thread
 	check(IsInRenderingThread());
 
-	FCachedRenderBuffers* Buffers = CachedBuffers.FindRef(&RenderHandle.Get());
-	if ( Buffers == nullptr )
-	{
-		// Rather than having a global pool, we associate the pools with a particular layout cacher.
-		// If we don't do this, all buffers eventually become as larger as the largest buffer, and it
-		// would be much better to keep the pools coherent with the sizes typically associated with
-		// a particular caching panel.
-		const ILayoutCache* LayoutCacher = RenderHandle->GetCacher();
-		TArray< FCachedRenderBuffers* >& Pool = CachedBufferPool.FindOrAdd(LayoutCacher);
-
-		// If the cached buffer pool is empty, time to create a new one!
-		if ( Pool.Num() == 0 )
-		{
-			Buffers = new FCachedRenderBuffers();
-			Buffers->VertexBuffer.Init(200);
-			Buffers->IndexBuffer.Init(200);
-		}
-		else
-		{
-			// If we found one in the pool, lets use it!
-			Buffers = Pool[0];
-			Pool.RemoveAtSwap(0, 1, false);
-		}
-
-		CachedBuffers.Add(&RenderHandle.Get(), Buffers);
-	}
+	FCachedRenderBuffers* Buffers = ResourceManager->FindOrCreateCachedBuffersForHandle(RenderHandle);
 
 	UpdateVertexAndIndexBuffers(RHICmdList, InBatchData, Buffers->VertexBuffer, Buffers->IndexBuffer);
 }
 
 void FSlateRHIRenderingPolicy::ReleaseCachedRenderData(FSlateRenderDataHandle* InRenderHandle)
 {
-	// Should only be called by the rendering thread
-	check(IsInRenderingThread());
-	check(InRenderHandle);
-
-	FCachedRenderBuffers* PooledBuffer = CachedBuffers.FindRef(InRenderHandle);
-	if ( ensure(PooledBuffer != nullptr) )
-	{
-		const ILayoutCache* LayoutCacher = InRenderHandle->GetCacher();
-		TArray< FCachedRenderBuffers* >* Pool = CachedBufferPool.Find(LayoutCacher);
-		if ( Pool )
-		{
-			Pool->Add(PooledBuffer);
-		}
-		else
-		{
-			// The buffer pool may have already been released, so lets just delete this buffer.
-			PooledBuffer->VertexBuffer.Destroy();
-			PooledBuffer->IndexBuffer.Destroy();
-			delete PooledBuffer;
-		}
-		
-		CachedBuffers.Remove(InRenderHandle);
-	}
+	ResourceManager->ReleaseCachedRenderData(InRenderHandle);
 }
 
 void FSlateRHIRenderingPolicy::ReleaseCachingResourcesFor(const ILayoutCache* Cacher)
 {
-	TArray< FCachedRenderBuffers* >* Pool = CachedBufferPool.Find(Cacher);
-	if ( Pool )
-	{
-		for ( FCachedRenderBuffers* PooledBuffer : *Pool )
-		{
-			PooledBuffer->VertexBuffer.Destroy();
-			PooledBuffer->IndexBuffer.Destroy();
-			delete PooledBuffer;
-		}
-
-		CachedBufferPool.Remove(Cacher);
-	}
+	ResourceManager->ReleaseCachingResourcesFor(Cacher);
 }
 
 void FSlateRHIRenderingPolicy::UpdateVertexAndIndexBuffers(FRHICommandListImmediate& RHICmdList, FSlateBatchData& InBatchData, TSlateElementVertexBuffer<FSlateVertex>& VertexBuffer, FSlateElementIndexBuffer& IndexBuffer)
@@ -613,7 +397,7 @@ void FSlateRHIRenderingPolicy::DrawElements(FRHICommandListImmediate& RHICmdList
 
 			if ( RenderHandle )
 			{
-				FCachedRenderBuffers* Buffers = CachedBuffers.FindRef(RenderHandle);
+				FCachedRenderBuffers* Buffers = ResourceManager->FindCachedBuffersForHandle(RenderHandle);
 				if ( Buffers != nullptr )
 				{
 					VertexBuffer = &Buffers->VertexBuffer;
@@ -677,9 +461,13 @@ void FSlateRHIRenderingPolicy::DrawElements(FRHICommandListImmediate& RHICmdList
 				GlobalVertexShader->SetVerticalAxisMultiplier(RHICmdList, bAllowSwitchVerticalAxis && RHINeedsToSwitchVerticalAxis(GShaderPlatformForFeatureLevel[GMaxRHIFeatureLevel]) ? -1.0f : 1.0f );
 #if !DEBUG_OVERDRAW
 				RHICmdList.SetBlendState(
-					(RenderBatch.DrawFlags & ESlateBatchDrawFlag::NoBlending)
+					( RenderBatch.DrawFlags & ESlateBatchDrawFlag::NoBlending )
 					? TStaticBlendState<>::GetRHI()
-					: TStaticBlendState<CW_RGBA, BO_Add, BF_SourceAlpha, BF_InverseSourceAlpha, BO_Add, BF_InverseDestAlpha, BF_One>::GetRHI()
+#if SLATE_PRE_MULTIPLY
+					: TStaticBlendState<CW_RGBA, BO_Add, BF_One, BF_InverseSourceAlpha, BO_Add, BF_One, BF_InverseSourceAlpha>::GetRHI()
+#else
+					: TStaticBlendState<CW_RGBA, BO_Add, BF_SourceAlpha, BF_InverseSourceAlpha, BO_Add, BF_One, BF_InverseSourceAlpha>::GetRHI()
+#endif
 					);
 #else
 				RHICmdList.SetBlendState(TStaticBlendState<CW_RGB, BO_Add, BF_One, BF_One, BO_Add, BF_Zero, BF_InverseSourceAlpha>::GetRHI());
@@ -853,83 +641,114 @@ void FSlateRHIRenderingPolicy::DrawElements(FRHICommandListImmediate& RHICmdList
 
 FSlateElementPS* FSlateRHIRenderingPolicy::GetTexturePixelShader( ESlateShader::Type ShaderType, ESlateDrawEffect::Type DrawEffects )
 {
-	FSlateElementPS* PixelShader = NULL;
+	FSlateElementPS* PixelShader = nullptr;
 	const auto FeatureLevel = GMaxRHIFeatureLevel;
 	auto ShaderMap = GetGlobalShaderMap(FeatureLevel);
 
+#define PixelShaderLookupTable(FullEffects, LimitedEffects)	\
+	switch ( ShaderType )																							\
+	{																							\
+	default:																							\
+	case ESlateShader::Default:																							\
+		PixelShader = *TShaderMapRef<TSlateElementPS<ESlateShader::Default, FullEffects> >(ShaderMap);									\
+		break;									\
+	case ESlateShader::Border:									\
+		PixelShader = *TShaderMapRef<TSlateElementPS<ESlateShader::Border, FullEffects> >(ShaderMap);									\
+		break;									\
+	case ESlateShader::Font:									\
+		PixelShader = *TShaderMapRef<TSlateElementPS<ESlateShader::Font, LimitedEffects> >(ShaderMap);									\
+		break;									\
+	case ESlateShader::LineSegment:									\
+		PixelShader = *TShaderMapRef<TSlateElementPS<ESlateShader::LineSegment, LimitedEffects> >(ShaderMap);									\
+		break;									\
+	}
+
 #if !DEBUG_OVERDRAW
-	const bool bDrawDisabled = (DrawEffects & ESlateDrawEffect::DisabledEffect) != 0;
-	const bool bUseTextureAlpha = (DrawEffects & ESlateDrawEffect::IgnoreTextureAlpha) == 0;
-	if( bDrawDisabled )
+	
+	const bool bPreMultipliedAlpha = ( DrawEffects & ESlateDrawEffect::PreMultipliedAlpha ) != 0;
+
+	if ( bPreMultipliedAlpha )
 	{
-		switch( ShaderType )
-		{
-		default:
-		case ESlateShader::Default:
-			if( bUseTextureAlpha )
-			{
-				PixelShader = *TShaderMapRef<TSlateElementPS<ESlateShader::Default, true, true> >(ShaderMap);
-			}
-			else
-			{
-				PixelShader = *TShaderMapRef<TSlateElementPS<ESlateShader::Default, true, false> >(ShaderMap);
-			}
-			break;
-		case ESlateShader::Border:
-			if( bUseTextureAlpha )
-			{
-				PixelShader = *TShaderMapRef<TSlateElementPS<ESlateShader::Border, true, true> >(ShaderMap);
-			}
-			else
-			{
-				PixelShader = *TShaderMapRef<TSlateElementPS<ESlateShader::Border, true, false> >(ShaderMap);
-			}
-			break;
-		case ESlateShader::Font:
-			PixelShader = *TShaderMapRef<TSlateElementPS<ESlateShader::Font, true> >(ShaderMap);
-			break;
-		case ESlateShader::LineSegment:
-			PixelShader = *TShaderMapRef<TSlateElementPS<ESlateShader::LineSegment, true> >(ShaderMap);
-			break;
-		}
+		PixelShader = *TShaderMapRef<FSlatePreMultiplyPassThroughPS>(ShaderMap);
 	}
 	else
 	{
-		switch( ShaderType )
+		const bool bDrawDisabled = ( DrawEffects & ESlateDrawEffect::DisabledEffect ) != 0;
+		const bool bUseTextureAlpha = ( DrawEffects & ESlateDrawEffect::IgnoreTextureAlpha ) == 0;
+
+		if ( bDrawDisabled )
 		{
-		default:
-		case ESlateShader::Default:
-			if( bUseTextureAlpha )
-			{
-				PixelShader = *TShaderMapRef<TSlateElementPS<ESlateShader::Default, false, true> >(ShaderMap);
-			}
-			else
-			{
-				PixelShader = *TShaderMapRef<TSlateElementPS<ESlateShader::Default, false, false> >(ShaderMap);
-			}
-			break;
-		case ESlateShader::Border:
-			if( bUseTextureAlpha )
+			switch ( ShaderType )
 			{
-				PixelShader = *TShaderMapRef<TSlateElementPS<ESlateShader::Border, false, true> >(ShaderMap);
+			default:
+			case ESlateShader::Default:
+				if ( bUseTextureAlpha )
+				{
+					PixelShader = *TShaderMapRef<TSlateElementPS<ESlateShader::Default, true, true> >(ShaderMap);
+				}
+				else
+				{
+					PixelShader = *TShaderMapRef<TSlateElementPS<ESlateShader::Default, true, false> >(ShaderMap);
+				}
+				break;
+			case ESlateShader::Border:
+				if ( bUseTextureAlpha )
+				{
+					PixelShader = *TShaderMapRef<TSlateElementPS<ESlateShader::Border, true, true> >(ShaderMap);
+				}
+				else
+				{
+					PixelShader = *TShaderMapRef<TSlateElementPS<ESlateShader::Border, true, false> >(ShaderMap);
+				}
+				break;
+			case ESlateShader::Font:
+				PixelShader = *TShaderMapRef<TSlateElementPS<ESlateShader::Font, true> >(ShaderMap);
+				break;
+			case ESlateShader::LineSegment:
+				PixelShader = *TShaderMapRef<TSlateElementPS<ESlateShader::LineSegment, true> >(ShaderMap);
+				break;
 			}
-			else
+		}
+		else
+		{
+			switch ( ShaderType )
 			{
-				PixelShader = *TShaderMapRef<TSlateElementPS<ESlateShader::Border, false, false> >(ShaderMap);
+			default:
+			case ESlateShader::Default:
+				if ( bUseTextureAlpha )
+				{
+					PixelShader = *TShaderMapRef<TSlateElementPS<ESlateShader::Default, false, true> >(ShaderMap);
+				}
+				else
+				{
+					PixelShader = *TShaderMapRef<TSlateElementPS<ESlateShader::Default, false, false> >(ShaderMap);
+				}
+				break;
+			case ESlateShader::Border:
+				if ( bUseTextureAlpha )
+				{
+					PixelShader = *TShaderMapRef<TSlateElementPS<ESlateShader::Border, false, true> >(ShaderMap);
+				}
+				else
+				{
+					PixelShader = *TShaderMapRef<TSlateElementPS<ESlateShader::Border, false, false> >(ShaderMap);
+				}
+				break;
+			case ESlateShader::Font:
+				PixelShader = *TShaderMapRef<TSlateElementPS<ESlateShader::Font, false> >(ShaderMap);
+				break;
+			case ESlateShader::LineSegment:
+				PixelShader = *TShaderMapRef<TSlateElementPS<ESlateShader::LineSegment, false> >(ShaderMap);
+				break;
 			}
-			break;
-		case ESlateShader::Font:
-			PixelShader = *TShaderMapRef<TSlateElementPS<ESlateShader::Font, false> >(ShaderMap);
-			break;
-		case ESlateShader::LineSegment:
-			PixelShader = *TShaderMapRef<TSlateElementPS<ESlateShader::LineSegment, false> >(ShaderMap);
-			break;
 		}
 	}
 #else
 	PixelShader = *TShaderMapRef<FSlateDebugOverdrawPS>(ShaderMap);
 #endif
 
+#undef PixelShaderLookupTable
+
 	return PixelShader;
 }
 
diff --git a/Engine/Source/Runtime/SlateRHIRenderer/Private/SlateRHIRenderingPolicy.h b/Engine/Source/Runtime/SlateRHIRenderer/Private/SlateRHIRenderingPolicy.h
--- a/Engine/Source/Runtime/SlateRHIRenderer/Private/SlateRHIRenderingPolicy.h
+++ b/Engine/Source/Runtime/SlateRHIRenderer/Private/SlateRHIRenderingPolicy.h
@@ -7,233 +7,6 @@ class FSlateRHIResourceManager;
 
 #include "RenderingPolicy.h"
 
-
-/** 
- * Vertex buffer containing all Slate vertices
- */
-template <typename VertexType>
-class TSlateElementVertexBuffer : public FVertexBuffer
-{
-public:
-	TSlateElementVertexBuffer()
-		: BufferSize(0)
-		, MinBufferSize(0)
-		, BufferUsageSize(0)
-	{}
-
-	~TSlateElementVertexBuffer() {};
-
-	void Init( int32 MinNumVertices )
-	{
-		MinBufferSize = sizeof(VertexType) * FMath::Max( MinNumVertices, 200 );
-
-		if ( IsInRenderingThread() )
-		{
-			InitResource();
-		}
-		else
-		{
-			BeginInitResource(this);
-		}
-	}
-
-	void Destroy()
-	{
-		if ( IsInRenderingThread() )
-		{
-			ReleaseResource();
-		}
-		else
-		{
-			BeginReleaseResource(this);
-		}
-	}
-
-	/** Initializes the vertex buffers RHI resource. */
-	virtual void InitDynamicRHI()
-	{
-		if( !IsValidRef(VertexBufferRHI) )
-		{
-			check( MinBufferSize > 0 )
-	
-			BufferSize = MinBufferSize;
-
-			FRHIResourceCreateInfo CreateInfo;
-			VertexBufferRHI = RHICreateVertexBuffer( MinBufferSize, BUF_Dynamic, CreateInfo );
-
-			// Ensure the vertex buffer could be created
-			check(IsValidRef(VertexBufferRHI));
-		}
-	}
-
-	/** Releases the vertex buffers RHI resource. */
-	virtual void ReleaseDynamicRHI()
-	{
-		VertexBufferRHI.SafeRelease();
-		BufferSize = 0;
-	}
-
-	/** Returns a friendly name for this buffer. */
-	virtual FString GetFriendlyName() const { return TEXT("SlateElementVertices"); }
-
-	/** Returns the size of the buffer in bytes. */
-	int32 GetBufferSize() const { return BufferSize; }
-
-	/** Returns the used size of this buffer */
-	int32 GetBufferUsageSize() const { return BufferUsageSize; }
-
-	/** Resets the usage of the buffer */
-	void ResetBufferUsage() { BufferUsageSize = 0; }
-
-	/** Resizes buffer, accumulates states safely on render thread */
-	void PreFillBuffer(int32 RequiredVertexCount, bool bShrinkToMinSize)
-	{
-		checkSlow(IsInRenderingThread());
-
-		if (RequiredVertexCount > 0 )
-		{
-#if !SLATE_USE_32BIT_INDICES
-			// make sure our index buffer can handle this
-			checkf(RequiredVertexCount < 0xFFFF, TEXT("Slate vertex buffer is too large (%d) to work with uint16 indices"), RequiredVertexCount);
-#endif
-			int32 RequiredBufferSize = RequiredVertexCount*sizeof(VertexType);
-
-			// resize if needed
-			if(RequiredBufferSize > GetBufferSize() || bShrinkToMinSize)
-			{
-				ResizeBuffer(RequiredBufferSize);
-			}
-
-			BufferUsageSize = RequiredBufferSize;
-		}
-
-	}
-
-	int32 GetMinBufferSize() const { return MinBufferSize; }
-
-	void* LockBuffer_RenderThread(int32 NumVertices)
-	{
-		uint32 RequiredBufferSize = NumVertices*sizeof(VertexType);	
-
-		return RHILockVertexBuffer( VertexBufferRHI, 0, RequiredBufferSize, RLM_WriteOnly );
-	}
-
-	void UnlockBuffer_RenderThread()
-	{
-		RHIUnlockVertexBuffer( VertexBufferRHI );
-	}
-
-	void* LockBuffer_RHIThread(int32 NumVertices)
-	{
-		int32 RequiredBufferSize = NumVertices*sizeof(VertexType);
-
-		//use non-flushing DynamicRHI version when on RHIThread.
-		return GDynamicRHI->RHILockVertexBuffer(VertexBufferRHI, 0, RequiredBufferSize, RLM_WriteOnly);
-	}
-
-	void UnlockBuffer_RHIThread()
-	{
-		GDynamicRHI->RHIUnlockVertexBuffer(VertexBufferRHI);		
-	}
-
-private:
-	/** Resizes the buffer to the passed in size.  Preserves internal data*/
-	void ResizeBuffer( int32 NewSizeBytes )
-	{
-		checkSlow( IsInRenderingThread() );
-
-		int32 FinalSize = FMath::Max( NewSizeBytes, MinBufferSize );
-
-		if( FinalSize != 0 && FinalSize != BufferSize )
-		{
-			VertexBufferRHI.SafeRelease();
-
-			FRHIResourceCreateInfo CreateInfo;
-			VertexBufferRHI = RHICreateVertexBuffer(FinalSize, BUF_Dynamic, CreateInfo);
-
-			check(IsValidRef(VertexBufferRHI));
-
-			BufferSize = FinalSize;
-		}
-	}
-
-private:
-	/** The size of the buffer in bytes. */
-	int32 BufferSize;
-	
-	/** The minimum size the buffer should always be */
-	int32 MinBufferSize;
-
-	/** The size of the used portion of the buffer */
-	int32 BufferUsageSize;
-
-	/** Hidden copy methods. */
-	TSlateElementVertexBuffer( const TSlateElementVertexBuffer& );
-	void operator=(const TSlateElementVertexBuffer& );
-};
-
-class FSlateElementIndexBuffer : public FIndexBuffer
-{
-public:
-	/**
-	 * Constructor
-	 * 
-	 * @param MinNumIndices	The minimum number of indices this buffer should always support
-	 */
-	FSlateElementIndexBuffer();
-	~FSlateElementIndexBuffer();
-
-	void Init( int32 MinNumIndices );
-	void Destroy();
-
-	/** Initializes the index buffers RHI resource. */
-	virtual void InitDynamicRHI() override;
-
-	/** Releases the index buffers RHI resource. */
-	virtual void ReleaseDynamicRHI() override;
-
-	/** Returns a friendly name for this buffer. */
-	virtual FString GetFriendlyName() const override { return TEXT("SlateElementIndices"); }
-
-	/** Returns the size of this buffer */
-	int32 GetBufferSize() const { return BufferSize; }
-
-	/** Returns the used size of this buffer */
-	int32 GetBufferUsageSize() const { return BufferUsageSize; }
-
-	/** Resets the usage of the buffer */
-	void ResetBufferUsage() { BufferUsageSize = 0; }
-
-	/** Resizes buffer, accumulates states safely on render thread */
-	void PreFillBuffer(int32 RequiredIndexCount, bool bShrinkToMinSize);
-
-	void* LockBuffer_RenderThread(int32 NumIndices);
-	void UnlockBuffer_RenderThread();
-
-	void* LockBuffer_RHIThread(int32 NumIndices);
-	void UnlockBuffer_RHIThread();
-
-	int32 GetMinBufferSize() const { return MinBufferSize; }
-
-private:
-	/** Resizes the buffer to the passed in size.  Preserves internal data */
-	void ResizeBuffer( int32 NewSizeBytes );
-private:
-	/** The current size of the buffer in bytes. */
-	int32 BufferSize;
-	
-	/** The minimum size the buffer should always be */
-	int32 MinBufferSize;
-
-	/** The size of the currently used portion of the buffer */
-	int32 BufferUsageSize;
-
-	/** Hidden copy methods. */
-	FSlateElementIndexBuffer( const FSlateElementIndexBuffer& );
-	void operator=(const FSlateElementIndexBuffer& );
-
-};
-
 class FSlateRHIRenderingPolicy : public FSlateRenderingPolicy
 {
 public:
@@ -259,8 +32,10 @@ public:
 	void EndDrawingWindows();
 
 	void SetUseGammaCorrection( bool bInUseGammaCorrection ) { bGammaCorrect = bInUseGammaCorrection; }
+
 protected:
 	void UpdateVertexAndIndexBuffers(FRHICommandListImmediate& RHICmdList, FSlateBatchData& BatchData, TSlateElementVertexBuffer<FSlateVertex>& VertexBuffer, FSlateElementIndexBuffer& IndexBuffer);
+
 private:
 	/**
 	 * Returns the pixel shader that should be used for the specified ShaderType and DrawEffects
@@ -285,24 +60,12 @@ private:
 	TSlateElementVertexBuffer<FSlateVertex> VertexBuffers[SlateRHIConstants::NumBuffers];
 	FSlateElementIndexBuffer IndexBuffers[SlateRHIConstants::NumBuffers];
 
-	struct FCachedRenderBuffers
-	{
-		TSlateElementVertexBuffer<FSlateVertex> VertexBuffer;
-		FSlateElementIndexBuffer IndexBuffer;
-	};
-
-	typedef TMap< FSlateRenderDataHandle*, FCachedRenderBuffers* > TCachedBufferMap;
-	TCachedBufferMap CachedBuffers;
-
-	typedef TMap< const ILayoutCache*, TArray< FCachedRenderBuffers* > > TCachedBufferPoolMap;
-	TCachedBufferPoolMap CachedBufferPool;
-
 	TSharedRef<FSlateRHIResourceManager> ResourceManager;
 	TSharedPtr<FSlateFontCache> FontCache;
 
-	uint8 CurrentBufferIndex;
-	
 	bool bGammaCorrect;
+
+	uint8 CurrentBufferIndex;
 };
 
 
diff --git a/Engine/Source/Runtime/SlateRHIRenderer/Private/SlateRHIResourceManager.cpp b/Engine/Source/Runtime/SlateRHIRenderer/Private/SlateRHIResourceManager.cpp
--- a/Engine/Source/Runtime/SlateRHIRenderer/Private/SlateRHIResourceManager.cpp
+++ b/Engine/Source/Runtime/SlateRHIRenderer/Private/SlateRHIResourceManager.cpp
@@ -588,7 +588,7 @@ FSlateShaderResourceProxy* FSlateRHIResourceManager::FindOrCreateDynamicTextureR
 	{
 		if ( UObject* ResourceObject = InBrush.GetResourceObject() )
 		{
-			if ( !(ResourceObject->IsA<UTexture2D>() || ResourceObject->IsA<UTexture2DDynamic>()) )
+			if ( !ResourceObject->IsA<UTexture>() )
 			{
 #if !(UE_BUILD_SHIPPING || UE_BUILD_TEST)
 				static TSet<UObject*> FailedTextures;
@@ -796,6 +796,81 @@ void FSlateRHIResourceManager::UpdateTextureAtlases()
 	}
 }
 
+FCachedRenderBuffers* FSlateRHIResourceManager::FindOrCreateCachedBuffersForHandle(const TSharedRef<FSlateRenderDataHandle, ESPMode::ThreadSafe>& RenderHandle)
+{
+	FCachedRenderBuffers* Buffers = CachedBuffers.FindRef(&RenderHandle.Get());
+	if ( Buffers == nullptr )
+	{
+		// Rather than having a global pool, we associate the pools with a particular layout cacher.
+		// If we don't do this, all buffers eventually become as larger as the largest buffer, and it
+		// would be much better to keep the pools coherent with the sizes typically associated with
+		// a particular caching panel.
+		const ILayoutCache* LayoutCacher = RenderHandle->GetCacher();
+		TArray< FCachedRenderBuffers* >& Pool = CachedBufferPool.FindOrAdd(LayoutCacher);
+
+		// If the cached buffer pool is empty, time to create a new one!
+		if ( Pool.Num() == 0 )
+		{
+			Buffers = new FCachedRenderBuffers();
+			Buffers->VertexBuffer.Init(200);
+			Buffers->IndexBuffer.Init(200);
+		}
+		else
+		{
+			// If we found one in the pool, lets use it!
+			Buffers = Pool[0];
+			Pool.RemoveAtSwap(0, 1, false);
+		}
+
+		CachedBuffers.Add(&RenderHandle.Get(), Buffers);
+	}
+
+	return Buffers;
+}
+
+void FSlateRHIResourceManager::ReleaseCachedRenderData(FSlateRenderDataHandle* InRenderHandle)
+{
+	// Should only be called by the rendering thread
+	check(IsInRenderingThread());
+	check(InRenderHandle);
+
+	FCachedRenderBuffers* PooledBuffer = CachedBuffers.FindRef(InRenderHandle);
+	if ( ensure(PooledBuffer != nullptr) )
+	{
+		const ILayoutCache* LayoutCacher = InRenderHandle->GetCacher();
+		TArray< FCachedRenderBuffers* >* Pool = CachedBufferPool.Find(LayoutCacher);
+		if ( Pool )
+		{
+			Pool->Add(PooledBuffer);
+		}
+		else
+		{
+			// The buffer pool may have already been released, so lets just delete this buffer.
+			PooledBuffer->VertexBuffer.Destroy();
+			PooledBuffer->IndexBuffer.Destroy();
+			delete PooledBuffer;
+		}
+
+		CachedBuffers.Remove(InRenderHandle);
+	}
+}
+
+void FSlateRHIResourceManager::ReleaseCachingResourcesFor(const ILayoutCache* Cacher)
+{
+	TArray< FCachedRenderBuffers* >* Pool = CachedBufferPool.Find(Cacher);
+	if ( Pool )
+	{
+		for ( FCachedRenderBuffers* PooledBuffer : *Pool )
+		{
+			PooledBuffer->VertexBuffer.Destroy();
+			PooledBuffer->IndexBuffer.Destroy();
+			delete PooledBuffer;
+		}
+
+		CachedBufferPool.Remove(Cacher);
+	}
+}
+
 void FSlateRHIResourceManager::ReleaseResources()
 {
 	checkSlow( IsThreadSafeForSlateRendering() );
@@ -812,6 +887,27 @@ void FSlateRHIResourceManager::ReleaseResources()
 
 	DynamicResourceMap.ReleaseResources();
 
+	for ( TCachedBufferMap::TIterator BufferIt(CachedBuffers); BufferIt; ++BufferIt )
+	{
+		FSlateRenderDataHandle* Handle = BufferIt.Key();
+		FCachedRenderBuffers* Buffer = BufferIt.Value();
+
+		Handle->Disconnect();
+
+		Buffer->VertexBuffer.Destroy();
+		Buffer->IndexBuffer.Destroy();
+	}
+
+	for ( TCachedBufferPoolMap::TIterator BufferIt(CachedBufferPool); BufferIt; ++BufferIt )
+	{
+		TArray< FCachedRenderBuffers* >& Pool = BufferIt.Value();
+		for ( FCachedRenderBuffers* PooledBuffer : Pool )
+		{
+			PooledBuffer->VertexBuffer.Destroy();
+			PooledBuffer->IndexBuffer.Destroy();
+		}
+	}
+
 	// Note the base class has texture proxies only which do not need to be released
 }
 
@@ -841,6 +937,25 @@ void FSlateRHIResourceManager::DeleteResources()
 
 	// Clean up mapping to texture
 	ClearTextureMap();
+
+	for ( TCachedBufferMap::TIterator BufferIt(CachedBuffers); BufferIt; ++BufferIt )
+	{
+		FCachedRenderBuffers* Buffer = BufferIt.Value();
+		delete Buffer;
+	}
+
+	CachedBuffers.Empty();
+
+	for ( TCachedBufferPoolMap::TIterator BufferIt(CachedBufferPool); BufferIt; ++BufferIt )
+	{
+		TArray< FCachedRenderBuffers* >& Pool = BufferIt.Value();
+		for ( FCachedRenderBuffers* PooledBuffer : Pool )
+		{
+			delete PooledBuffer;
+		}
+	}
+
+	CachedBufferPool.Empty();
 }
 
 void FSlateRHIResourceManager::ReloadTextures()
diff --git a/Engine/Source/Runtime/SlateRHIRenderer/Private/SlateRHIResourceManager.h b/Engine/Source/Runtime/SlateRHIRenderer/Private/SlateRHIResourceManager.h
--- a/Engine/Source/Runtime/SlateRHIRenderer/Private/SlateRHIResourceManager.h
+++ b/Engine/Source/Runtime/SlateRHIRenderer/Private/SlateRHIResourceManager.h
@@ -3,6 +3,8 @@
 #pragma once
 
 #include "TextureAtlas.h"
+#include "SlateElementIndexBuffer.h"
+#include "SlateElementVertexBuffer.h"
 
 class FSlateDynamicTextureResource;
 class FSlateUTextureResource;
@@ -61,6 +63,13 @@ private:
 };
 
 
+struct FCachedRenderBuffers
+{
+	TSlateElementVertexBuffer<FSlateVertex> VertexBuffer;
+	FSlateElementIndexBuffer IndexBuffer;
+};
+
+
 /**
  * Stores a mapping of texture names to their RHI texture resource               
  */
@@ -143,6 +152,14 @@ public:
 	virtual bool LoadTexture( const FName& TextureName, const FString& ResourcePath, uint32& Width, uint32& Height, TArray<uint8>& DecodedImage );
 	virtual bool LoadTexture( const FSlateBrush& InBrush, uint32& Width, uint32& Height, TArray<uint8>& DecodedImage );
 
+	FCachedRenderBuffers* FindCachedBuffersForHandle(const FSlateRenderDataHandle* RenderHandle) const
+	{
+		return CachedBuffers.FindRef(RenderHandle);
+	}
+
+	FCachedRenderBuffers* FindOrCreateCachedBuffersForHandle(const TSharedRef<FSlateRenderDataHandle, ESPMode::ThreadSafe>& RenderHandle);
+	void ReleaseCachedRenderData(FSlateRenderDataHandle* InRenderHandle);
+	void ReleaseCachingResourcesFor(const ILayoutCache* Cacher);
 
 	/**
 	 * Releases rendering resources
@@ -224,5 +241,11 @@ private:
 	FIntPoint MaxAltasedTextureSize;
 	/** Needed for displaying an error texture when we end up with bad resources. */
 	UTexture* BadResourceTexture;
+
+	typedef TMap< FSlateRenderDataHandle*, FCachedRenderBuffers* > TCachedBufferMap;
+	TCachedBufferMap CachedBuffers;
+
+	typedef TMap< const ILayoutCache*, TArray< FCachedRenderBuffers* > > TCachedBufferPoolMap;
+	TCachedBufferPoolMap CachedBufferPool;
 };
 
diff --git a/Engine/Source/Runtime/SlateRHIRenderer/Private/SlateShaders.cpp b/Engine/Source/Runtime/SlateRHIRenderer/Private/SlateShaders.cpp
--- a/Engine/Source/Runtime/SlateRHIRenderer/Private/SlateShaders.cpp
+++ b/Engine/Source/Runtime/SlateRHIRenderer/Private/SlateShaders.cpp
@@ -11,26 +11,28 @@ IMPLEMENT_SHADER_TYPE(, FSlateElementVS, TEXT("SlateVertexShader"),TEXT("Main"),
 
 IMPLEMENT_SHADER_TYPE(, FSlateDebugOverdrawPS, TEXT("SlateElementPixelShader"), TEXT("DebugOverdrawMain"), SF_Pixel );
 
+IMPLEMENT_SHADER_TYPE(, FSlatePreMultiplyPassThroughPS, TEXT("SlateElementPixelShader"), TEXT("Main"), SF_Pixel);
+
 #define IMPLEMENT_SLATE_PIXELSHADER_TYPE(ShaderType, bDrawDisabledEffect, bUseTextureAlpha) \
-	typedef TSlateElementPS<ESlateShader::ShaderType,bDrawDisabledEffect,bUseTextureAlpha> TSlateElementPS##ShaderType##bDrawDisabledEffect##bUseTextureAlpha; \
-	IMPLEMENT_SHADER_TYPE(template<>,TSlateElementPS##ShaderType##bDrawDisabledEffect##bUseTextureAlpha,TEXT("SlateElementPixelShader"),TEXT("Main"),SF_Pixel);
+	typedef TSlateElementPS<ESlateShader::ShaderType,bDrawDisabledEffect,bUseTextureAlpha> TSlateElementPS##ShaderType##bDrawDisabledEffect##bUseTextureAlpha##PMA; \
+	IMPLEMENT_SHADER_TYPE(template<>,TSlateElementPS##ShaderType##bDrawDisabledEffect##bUseTextureAlpha##PMA,TEXT("SlateElementPixelShader"),TEXT("Main"),SF_Pixel);
 
 /**
- * All the different permutations of shaders used by slate. Uses #defines to avoid dynamic branches                 
- */
-IMPLEMENT_SLATE_PIXELSHADER_TYPE(Default,		false, true);
-IMPLEMENT_SLATE_PIXELSHADER_TYPE(Border,		false, true);
-IMPLEMENT_SLATE_PIXELSHADER_TYPE(Default,		true, true);
-IMPLEMENT_SLATE_PIXELSHADER_TYPE(Border,		true, true);
-IMPLEMENT_SLATE_PIXELSHADER_TYPE(Default,		false, false);
-IMPLEMENT_SLATE_PIXELSHADER_TYPE(Border,		false, false);
-IMPLEMENT_SLATE_PIXELSHADER_TYPE(Default,		true, false);
-IMPLEMENT_SLATE_PIXELSHADER_TYPE(Border,		true, false);
-
-IMPLEMENT_SLATE_PIXELSHADER_TYPE(Font,			false, true);
-IMPLEMENT_SLATE_PIXELSHADER_TYPE(LineSegment,	false, true);
-IMPLEMENT_SLATE_PIXELSHADER_TYPE(Font,			true, true);
-IMPLEMENT_SLATE_PIXELSHADER_TYPE(LineSegment,	true, true);
+* All the different permutations of shaders used by slate. Uses #defines to avoid dynamic branches
+*/
+IMPLEMENT_SLATE_PIXELSHADER_TYPE(Default, false, true);
+IMPLEMENT_SLATE_PIXELSHADER_TYPE(Border, false, true);
+IMPLEMENT_SLATE_PIXELSHADER_TYPE(Default, true, true);
+IMPLEMENT_SLATE_PIXELSHADER_TYPE(Border, true, true);
+IMPLEMENT_SLATE_PIXELSHADER_TYPE(Default, false, false);
+IMPLEMENT_SLATE_PIXELSHADER_TYPE(Border, false, false);
+IMPLEMENT_SLATE_PIXELSHADER_TYPE(Default, true, false);
+IMPLEMENT_SLATE_PIXELSHADER_TYPE(Border, true, false);
+
+IMPLEMENT_SLATE_PIXELSHADER_TYPE(Font, false, true);
+IMPLEMENT_SLATE_PIXELSHADER_TYPE(LineSegment, false, true);
+IMPLEMENT_SLATE_PIXELSHADER_TYPE(Font, true, true);
+IMPLEMENT_SLATE_PIXELSHADER_TYPE(LineSegment, true, true);
 
 /** The Slate vertex declaration. */
 TGlobalResource<FSlateVertexDeclaration> GSlateVertexDeclaration;
diff --git a/Engine/Source/Runtime/SlateRHIRenderer/Private/SlateShaders.h b/Engine/Source/Runtime/SlateRHIRenderer/Private/SlateShaders.h
--- a/Engine/Source/Runtime/SlateRHIRenderer/Private/SlateShaders.h
+++ b/Engine/Source/Runtime/SlateRHIRenderer/Private/SlateShaders.h
@@ -8,6 +8,7 @@
 #include "ShaderParameters.h"
 #include "UniformBuffer.h"
 #include "ShaderParameterUtils.h"
+#include "Rendering/RenderingCommon.h"
 
 extern uint32 GSlateShaderColorVisionDeficiencyType;
 
@@ -137,7 +138,7 @@ public:
 
 	/**
 	 * Sets the display gamma.
- 	 *
+	 *
 	 * @param DisplayGamma The display gamma to use
 	 */
 	void SetDisplayGamma(FRHICommandList& RHICmdList, float InDisplayGamma)
@@ -169,7 +170,7 @@ private:
 /** 
  * Pixel shader types for all elements
  */
-template<ESlateShader::Type ShaderType,bool bDrawDisabledEffect,bool bUseTextureAlpha=true> 
+template<ESlateShader::Type ShaderType, bool bDrawDisabledEffect, bool bUseTextureAlpha=true>
 class TSlateElementPS : public FSlateElementPS
 {
 	DECLARE_SHADER_TYPE( TSlateElementPS, Global );
@@ -193,10 +194,11 @@ public:
 	{
 		// Set defines based on what this shader will be used for
 		OutEnvironment.SetDefine(TEXT("SHADER_TYPE"), (uint32)ShaderType);
-		OutEnvironment.SetDefine(TEXT("DRAW_DISABLED_EFFECT"), (uint32)(bDrawDisabledEffect ? 1 : 0));
-		OutEnvironment.SetDefine(TEXT("USE_TEXTURE_ALPHA"), (uint32)(bUseTextureAlpha ? 1 : 0));
+		OutEnvironment.SetDefine(TEXT("DRAW_DISABLED_EFFECT"), (uint32)( bDrawDisabledEffect ? 1 : 0 ));
+		OutEnvironment.SetDefine(TEXT("USE_TEXTURE_ALPHA"), (uint32)( bUseTextureAlpha ? 1 : 0 ));
 		OutEnvironment.SetDefine(TEXT("COLOR_VISION_DEFICIENCY_TYPE"), GSlateShaderColorVisionDeficiencyType);
 		OutEnvironment.SetDefine(TEXT("USE_MATERIALS"), (uint32)0);
+		OutEnvironment.SetDefine(TEXT("PRE_MULTIPLIED_ALPHA"), (uint32)0);
 
 		FSlateElementPS::ModifyCompilationEnvironment( Platform, OutEnvironment );
 	}
@@ -210,6 +212,33 @@ private:
 
 };
 
+/** Simple passthrough for Premultiplied alpha */
+class FSlatePreMultiplyPassThroughPS : public TSlateElementPS<ESlateShader::Default, false, true>
+{
+	DECLARE_SHADER_TYPE(FSlatePreMultiplyPassThroughPS, Global);
+public:
+
+	FSlatePreMultiplyPassThroughPS()
+	{
+	}
+
+	/** Constructor.  Binds all parameters used by the shader */
+	FSlatePreMultiplyPassThroughPS(const ShaderMetaType::CompiledShaderInitializerType& Initializer)
+		: TSlateElementPS<ESlateShader::Default, false, true>(Initializer)
+	{
+	}
+
+	/**
+	* Modifies the compilation of this shader
+	*/
+	static void ModifyCompilationEnvironment(EShaderPlatform Platform, FShaderCompilerEnvironment& OutEnvironment)
+	{
+		TSlateElementPS<ESlateShader::Default, false, true>::ModifyCompilationEnvironment(Platform, OutEnvironment);
+
+		OutEnvironment.SetDefine(TEXT("PRE_MULTIPLIED_ALPHA"), (uint32)1);
+	}
+};
+
 /** 
  * Pixel shader for debugging Slate overdraw
  */
diff --git a/Engine/Source/Runtime/SlateRHIRenderer/Public/Interfaces/ISlate3DRenderer.h b/Engine/Source/Runtime/SlateRHIRenderer/Public/Interfaces/ISlate3DRenderer.h
--- a/Engine/Source/Runtime/SlateRHIRenderer/Public/Interfaces/ISlate3DRenderer.h
+++ b/Engine/Source/Runtime/SlateRHIRenderer/Public/Interfaces/ISlate3DRenderer.h
@@ -5,12 +5,12 @@
 class ISlate3DRenderer 
 {
 public:
-	virtual ~ISlate3DRenderer() {};
+	virtual ~ISlate3DRenderer() {}
 
 	/** @return The free buffer for drawing */
 	virtual FSlateDrawBuffer& GetDrawBuffer() = 0;
 
- 	/** 
+	/** 
 	 * Batches the draw elements in the buffer to prepare it for rendering.
 	 * Call in the game thread before sending to the render thread.
 	 *
diff --git a/Engine/Source/Runtime/SlateRHIRenderer/Public/Interfaces/ISlateRHIRendererModule.h b/Engine/Source/Runtime/SlateRHIRenderer/Public/Interfaces/ISlateRHIRendererModule.h
--- a/Engine/Source/Runtime/SlateRHIRenderer/Public/Interfaces/ISlateRHIRendererModule.h
+++ b/Engine/Source/Runtime/SlateRHIRenderer/Public/Interfaces/ISlateRHIRendererModule.h
@@ -29,7 +29,7 @@ public:
 	 */
 	virtual TSharedRef<FSlateRenderer> CreateSlateRHIRenderer( ) = 0;
 
-	virtual TSharedRef<ISlate3DRenderer> CreateSlate3DRenderer() = 0;
+	virtual TSharedRef<ISlate3DRenderer> CreateSlate3DRenderer(bool bUseGammaCorrection) = 0;
 
 	virtual TSharedRef<ISlateFontAtlasFactory> CreateSlateFontAtlasFactory() = 0;
 
diff --git a/Engine/Source/Runtime/UMG/Private/Components/InvalidationBox.cpp b/Engine/Source/Runtime/UMG/Private/Components/InvalidationBox.cpp
--- a/Engine/Source/Runtime/UMG/Private/Components/InvalidationBox.cpp
+++ b/Engine/Source/Runtime/UMG/Private/Components/InvalidationBox.cpp
@@ -29,14 +29,7 @@ TSharedRef<SWidget> UInvalidationBox::RebuildWidget()
 		SNew(SInvalidationPanel)
 		.CacheRelativeTransforms(CacheRelativeTransforms);
 
-	if ( IsDesignTime() )
-	{
-		MyInvalidationPanel->SetCanCache(false);
-	}
-	else
-	{
-		MyInvalidationPanel->SetCanCache(bCanCache);
-	}
+	MyInvalidationPanel->SetCanCache(IsDesignTime() ? false : bCanCache);
 
 	if ( GetChildrenCount() > 0 )
 	{
diff --git a/Engine/Source/Runtime/UMG/Private/Components/RetainerBox.cpp b/Engine/Source/Runtime/UMG/Private/Components/RetainerBox.cpp
new file mode 100644
--- /dev/null
+++ b/Engine/Source/Runtime/UMG/Private/Components/RetainerBox.cpp
@@ -0,0 +1,79 @@
+// Copyright 1998-2015 Epic Games, Inc. All Rights Reserved.
+
+#include "UMGPrivatePCH.h"
+
+#include "RetainerBox.h"
+#include "SRetainerWidget.h"
+
+#define LOCTEXT_NAMESPACE "UMG"
+
+/////////////////////////////////////////////////////
+// URetainerBox
+
+URetainerBox::URetainerBox(const FObjectInitializer& ObjectInitializer)
+	: Super(ObjectInitializer)
+{
+	Visibility = ESlateVisibility::Visible;
+	Phase = 0;
+	PhaseCount = 1;
+}
+
+void URetainerBox::ReleaseSlateResources(bool bReleaseChildren)
+{
+	Super::ReleaseSlateResources(bReleaseChildren);
+
+	MyRetainerWidget.Reset();
+}
+
+TSharedRef<SWidget> URetainerBox::RebuildWidget()
+{
+	MyRetainerWidget =
+		SNew(SRetainerWidget)
+		.Phase(Phase)
+		.PhaseCount(PhaseCount);
+
+	MyRetainerWidget->SetRetainedRendering(IsDesignTime() ? false : true);
+
+	if ( GetChildrenCount() > 0 )
+	{
+		MyRetainerWidget->SetContent(GetContentSlot()->Content ? GetContentSlot()->Content->TakeWidget() : SNullWidget::NullWidget);
+	}
+	
+	return BuildDesignTimeWidget(MyRetainerWidget.ToSharedRef());
+}
+
+void URetainerBox::OnSlotAdded(UPanelSlot* Slot)
+{
+	// Add the child to the live slot if it already exists
+	if ( MyRetainerWidget.IsValid() )
+	{
+		MyRetainerWidget->SetContent(Slot->Content ? Slot->Content->TakeWidget() : SNullWidget::NullWidget);
+	}
+}
+
+void URetainerBox::OnSlotRemoved(UPanelSlot* Slot)
+{
+	// Remove the widget from the live slot if it exists.
+	if ( MyRetainerWidget.IsValid() )
+	{
+		MyRetainerWidget->SetContent(SNullWidget::NullWidget);
+	}
+}
+
+#if WITH_EDITOR
+
+const FSlateBrush* URetainerBox::GetEditorIcon()
+{
+	return FUMGStyle::Get().GetBrush("Widget.MenuAnchor");
+}
+
+const FText URetainerBox::GetPaletteCategory()
+{
+	return LOCTEXT("Optimization", "Optimization");
+}
+
+#endif
+
+/////////////////////////////////////////////////////
+
+#undef LOCTEXT_NAMESPACE
diff --git a/Engine/Source/Runtime/UMG/Private/Components/WidgetComponent.cpp b/Engine/Source/Runtime/UMG/Private/Components/WidgetComponent.cpp
--- a/Engine/Source/Runtime/UMG/Private/Components/WidgetComponent.cpp
+++ b/Engine/Source/Runtime/UMG/Private/Components/WidgetComponent.cpp
@@ -655,11 +655,14 @@ void UWidgetComponent::TickComponent(float DeltaTime, enum ELevelTick TickType,
 			if ( TickWhenOffscreen || GetWorld()->TimeSince(LastRenderTime) <= RenderTimeThreshold )
 			{
 				UpdateRenderTarget();
+				
+				const float DrawScale = 1.0f;
 
 				WidgetRenderer->DrawWindow(
 					RenderTarget,
 					HitTestGrid.ToSharedRef(),
 					SlateWidget.ToSharedRef(),
+					DrawScale,
 					DrawSize,
 					DeltaTime);
 			}
