Add missing resource transition to slate

diff --git a/Engine/Source/Runtime/SlateRHIRenderer/Private/Slate3DRenderer.cpp b/Engine/Source/Runtime/SlateRHIRenderer/Private/Slate3DRenderer.cpp
--- a/Engine/Source/Runtime/SlateRHIRenderer/Private/Slate3DRenderer.cpp
+++ b/Engine/Source/Runtime/SlateRHIRenderer/Private/Slate3DRenderer.cpp
@@ -106,6 +106,7 @@ void FSlate3DRenderer::DrawWindowToTarget_RenderThread( FRHICommandListImmediate
 	Info.bClearColor = true;
 	ensure(ColorRTV.Texture->GetClearColor() == RenderTarget->ClearColor);
 
+	RHICmdList.TransitionResource(EResourceTransitionAccess::EWritable, RTResource);
 	RHICmdList.SetRenderTargetsAndClear(Info);
 
 	FMatrix ProjectionMatrix = FSlateRHIRenderer::CreateProjectionMatrix( RenderTarget->SizeX, RenderTarget->SizeY );
