Merging //UE4/Release-4.11 to //UE4/Main (Up to CL#2897738)

diff --git a/Engine/Source/Runtime/Engine/Private/SceneView.cpp b/Engine/Source/Runtime/Engine/Private/SceneView.cpp
--- a/Engine/Source/Runtime/Engine/Private/SceneView.cpp
+++ b/Engine/Source/Runtime/Engine/Private/SceneView.cpp
@@ -588,6 +588,12 @@ FSceneView::FSceneView(const FSceneViewInitOptions& InitOptions)
 	// Query instanced stereo state
 	static const auto CVar = IConsoleManager::Get().FindTConsoleVariableDataInt(TEXT("vr.InstancedStereo"));
 	bIsInstancedStereoEnabled = (ShaderPlatform == EShaderPlatform::SP_PCD3D_SM5 || ShaderPlatform == EShaderPlatform::SP_PS4) ? (CVar ? (CVar->GetValueOnAnyThread() != false) : false) : false;
+
+	// Check to see if this is an isolated capture component (i.e. not part of a stereo pair), and disable the instancing if so.  This will prevent single views that are set up to emulate different eyes from going down a path where it's paired view is not available (see StereoPanorama capturing)
+	if (Family && (Family->Views.Num() < 2))
+	{
+		bIsInstancedStereoEnabled = false;
+	}
 }
 
 static TAutoConsoleVariable<int32> CVarCompensateForFOV(
@@ -1944,8 +1950,13 @@ EDebugViewShaderMode FSceneViewFamily::ChooseDebugViewShaderMode() const
 
 const FSceneView& FSceneViewFamily::GetStereoEyeView(const EStereoscopicPass Eye) const
 {
-	const int32 EyeIndex = static_cast<int32>(Eye);
-	check(Views.Num() > 0 && Views.Num() >= EyeIndex);
+	int32 EyeIndex = static_cast<int32>(Eye);
+
+	// If we have less views than our eye index, it implies that we're using a scene capture that is emulating stereo behaviour, so re-index accordingly
+	if (Views.Num() <= EyeIndex)
+	{
+		EyeIndex = 1;
+	}
 
 	// Mono or left eye
 	if (EyeIndex <= 1)
diff --git a/Engine/Source/Runtime/Engine/Private/ShaderCompiler/ShaderCompiler.cpp b/Engine/Source/Runtime/Engine/Private/ShaderCompiler/ShaderCompiler.cpp
--- a/Engine/Source/Runtime/Engine/Private/ShaderCompiler/ShaderCompiler.cpp
+++ b/Engine/Source/Runtime/Engine/Private/ShaderCompiler/ShaderCompiler.cpp
@@ -824,7 +824,10 @@ void FShaderCompileThreadRunnable::WriteNewTasks()
 
 			if (!DoWriteTasks(CurrentWorkerInfo.QueuedJobs, *TransferFile))
 			{
-				UE_LOG(LogShaderCompilers, Fatal, TEXT("Could not write the shader compiler transfer filename to '%s'."), *TransferFileName);
+				uint64 TotalDiskSpace = 0;
+				uint64 FreeDiskSpace = 0;
+				FPlatformMisc::GetDiskTotalAndFreeSpace(TransferFileName, TotalDiskSpace, FreeDiskSpace);
+				UE_LOG(LogShaderCompilers, Fatal, TEXT("Could not write the shader compiler transfer filename to '%s' (Free Disk Space: %llu."), *TransferFileName, FreeDiskSpace);
 			}
 			delete TransferFile;
 
@@ -832,7 +835,10 @@ void FShaderCompileThreadRunnable::WriteNewTasks()
 			FString ProperTransferFileName = WorkingDirectory / TEXT("WorkerInputOnly.in");
 			if (!IFileManager::Get().Move(*ProperTransferFileName, *TransferFileName))
 			{
-				UE_LOG(LogShaderCompilers, Fatal, TEXT("Could not rename the shader compiler transfer filename to '%s' from '%s'."), *ProperTransferFileName, *TransferFileName);
+				uint64 TotalDiskSpace = 0;
+				uint64 FreeDiskSpace = 0;
+				FPlatformMisc::GetDiskTotalAndFreeSpace(TransferFileName, TotalDiskSpace, FreeDiskSpace);
+				UE_LOG(LogShaderCompilers, Fatal, TEXT("Could not rename the shader compiler transfer filename to '%s' from '%s' (Free Disk Space: %llu)."), *ProperTransferFileName, *TransferFileName, FreeDiskSpace);
 			}
 		}
 	}
diff --git a/Engine/Source/Runtime/Engine/Private/Slate/SceneViewport.cpp b/Engine/Source/Runtime/Engine/Private/Slate/SceneViewport.cpp
--- a/Engine/Source/Runtime/Engine/Private/Slate/SceneViewport.cpp
+++ b/Engine/Source/Runtime/Engine/Private/Slate/SceneViewport.cpp
@@ -384,8 +384,8 @@ FReply FSceneViewport::OnMouseButtonDown( const FGeometry& InGeometry, const FPo
 			ViewportClient->CaptureMouseOnClick() == EMouseCaptureMode::CaptureDuringMouseDown ||
 			(ViewportClient->CaptureMouseOnClick() == EMouseCaptureMode::CaptureDuringRightMouseDown && InMouseEvent.GetEffectingButton() == EKeys::RightMouseButton);
 
-		// Capture if we aren't currently a game viewport, or if we've been set to IgnoreCaptureClick or we already have capture.
-		const bool bProcessInputPrimary = !IsCurrentlyGameViewport() || !bIgnoreCaptureClick || HasMouseCapture();
+		// Process primary input if we aren't currently a game viewport, we already have capture, or we are permanent capture that dosn't consume the mouse down.
+		const bool bProcessInputPrimary = !IsCurrentlyGameViewport() || !bIgnoreCaptureClick || HasMouseCapture() || (ViewportClient->CaptureMouseOnClick() == EMouseCaptureMode::CapturePermanently_IncludingInitialMouseDown);
 
 		const bool bAnyMenuWasVisible = FSlateApplication::Get().AnyMenusVisible();
 
@@ -401,9 +401,13 @@ FReply FSceneViewport::OnMouseButtonDown( const FGeometry& InGeometry, const FPo
 		// a new menu was opened if there was previously not a menu visible but now there is
 		const bool bNewMenuWasOpened = !bAnyMenuWasVisible && FSlateApplication::Get().AnyMenusVisible();
 
+		const bool bPermanentCapture =
+			(ViewportClient->CaptureMouseOnClick() == EMouseCaptureMode::CapturePermanently) ||
+			(ViewportClient->CaptureMouseOnClick() == EMouseCaptureMode::CapturePermanently_IncludingInitialMouseDown);
+
 		if (FSlateApplication::Get().IsActive() && !ViewportClient->IgnoreInput() &&
 			!bNewMenuWasOpened && // We should not focus the viewport if a menu was opened as it would close the menu
-			(ViewportClient->CaptureMouseOnClick() == EMouseCaptureMode::CapturePermanently || bTemporaryCapture))
+			(bPermanentCapture || bTemporaryCapture))
 		{
 			CurrentReplyState = AcquireFocusAndCapture(FIntPoint(InMouseEvent.GetScreenSpacePosition().X, InMouseEvent.GetScreenSpacePosition().Y));
 		}
diff --git a/Engine/Source/Runtime/Engine/Private/UnrealEngine.cpp b/Engine/Source/Runtime/Engine/Private/UnrealEngine.cpp
--- a/Engine/Source/Runtime/Engine/Private/UnrealEngine.cpp
+++ b/Engine/Source/Runtime/Engine/Private/UnrealEngine.cpp
@@ -9654,6 +9654,25 @@ bool UEngine::LoadMap( FWorldContext& WorldContext, FURL URL, class UPendingNetG
 	UPackage* WorldPackage = NULL;
 	UWorld*	NewWorld = NULL;
 	
+	// If this world is a PIE instance, we need to check if we are travelling to another PIE instance's world.
+	// If we are, we need to set the PIERemapPrefix so that we load a copy of that world, instead of loading the
+	// PIE world directly.
+	if (!WorldContext.PIEPrefix.IsEmpty())
+	{
+		for (const FWorldContext& WorldContextFromList : WorldList)
+		{
+			// We want to ignore our own PIE instance so that we don't unnecessarily set the PIERemapPrefix if we are not travelling to
+			// a server.
+			if (WorldContextFromList.World() != WorldContext.World())
+			{
+				if (!WorldContextFromList.PIEPrefix.IsEmpty() && URL.Map.Contains(WorldContextFromList.PIEPrefix))
+				{
+					WorldContext.PIERemapPrefix = WorldContextFromList.PIEPrefix;
+				}
+			}
+		}
+	}
+
 	// Is this a PIE networking thing?
 	if (!WorldContext.PIERemapPrefix.IsEmpty())
 	{
diff --git a/Engine/Source/Runtime/Engine/Private/UserInterface/InputSettings.cpp b/Engine/Source/Runtime/Engine/Private/UserInterface/InputSettings.cpp
--- a/Engine/Source/Runtime/Engine/Private/UserInterface/InputSettings.cpp
+++ b/Engine/Source/Runtime/Engine/Private/UserInterface/InputSettings.cpp
@@ -12,7 +12,8 @@
 #include "GameFramework/InputSettings.h"
 
 UInputSettings::UInputSettings(const FObjectInitializer& ObjectInitializer)
-: Super(ObjectInitializer)
+	: Super(ObjectInitializer)
+	, DefaultViewportMouseCaptureMode(EMouseCaptureMode::CapturePermanently_IncludingInitialMouseDown)
 {
 }
 
diff --git a/Engine/Source/Runtime/Engine/Private/UserInterface/PlayerInput.cpp b/Engine/Source/Runtime/Engine/Private/UserInterface/PlayerInput.cpp
--- a/Engine/Source/Runtime/Engine/Private/UserInterface/PlayerInput.cpp
+++ b/Engine/Source/Runtime/Engine/Private/UserInterface/PlayerInput.cpp
@@ -306,6 +306,22 @@ bool UPlayerInput::InputTouch(uint32 Handle, ETouchType::Type Type, const FVecto
 	{
 	case ETouchType::Began:
 		KeyState.EventAccumulator[IE_Pressed].Add(++EventCount);
+		if (KeyState.bDownPrevious == false)
+		{
+			UWorld* World = GetWorld();
+			check(World);
+
+			// check for doubleclick
+			// note, a tripleclick will currently count as a 2nd double click.
+			const float WorldRealTimeSeconds = World->GetRealTimeSeconds();
+			if ((WorldRealTimeSeconds - KeyState.LastUpDownTransitionTime) < GetDefault<UInputSettings>()->DoubleClickTime)
+			{
+				KeyState.EventAccumulator[IE_DoubleClick].Add(++EventCount);
+			}
+
+			// just went down
+			KeyState.LastUpDownTransitionTime = WorldRealTimeSeconds;
+		}
 		break;
 
 	case ETouchType::Ended:
diff --git a/Engine/Source/Runtime/Engine/Private/Vehicles/WheeledVehicleMovementComponent.cpp b/Engine/Source/Runtime/Engine/Private/Vehicles/WheeledVehicleMovementComponent.cpp
--- a/Engine/Source/Runtime/Engine/Private/Vehicles/WheeledVehicleMovementComponent.cpp
+++ b/Engine/Source/Runtime/Engine/Private/Vehicles/WheeledVehicleMovementComponent.cpp
@@ -634,6 +634,12 @@ void UWheeledVehicleMovementComponent::CreatePhysicsState()
 				//LogVehicleSettings( PVehicle );
 				SCOPED_SCENE_WRITE_LOCK(VehicleManager->GetScene());
 				PVehicle->getRigidDynamicActor()->wakeUp();
+
+				// Need to bind to the notify delegate on the mesh incase physics state is changed
+				if(USkeletalMeshComponent* MeshComp = Cast<USkeletalMeshComponent>(GetMesh()))
+				{
+					MeshOnPhysicsStateChangeHandle = MeshComp->RegisterOnPhysicsCreatedDelegate(FOnSkelMeshPhysicsCreated::CreateUObject(this, &UWheeledVehicleMovementComponent::RecreatePhysicsState));
+				}
 			}
 		}
 	}
@@ -650,6 +656,14 @@ void UWheeledVehicleMovementComponent::DestroyPhysicsState()
 		World->GetPhysicsScene()->GetVehicleManager()->RemoveVehicle( this );
 		PVehicle = NULL;
 
+		if(MeshOnPhysicsStateChangeHandle.IsValid())
+		{
+			if(USkeletalMeshComponent* MeshComp = Cast<USkeletalMeshComponent>(GetMesh()))
+			{
+				MeshComp->UnregisterOnPhysicsCreatedDelegate(MeshOnPhysicsStateChangeHandle);
+			}
+		}
+
 		if ( UpdatedComponent )
 		{
 			UpdatedComponent->RecreatePhysicsState();
diff --git a/Engine/Source/Runtime/Engine/Private/World.cpp b/Engine/Source/Runtime/Engine/Private/World.cpp
--- a/Engine/Source/Runtime/Engine/Private/World.cpp
+++ b/Engine/Source/Runtime/Engine/Private/World.cpp
@@ -2870,7 +2870,7 @@ bool UWorld::HandleLogActorCountsCommand( const TCHAR* Cmd, FOutputDevice& Ar, U
 
 bool UWorld::HandleDemoRecordCommand( const TCHAR* Cmd, FOutputDevice& Ar, UWorld* InWorld )
 {
-	if ( InWorld != nullptr && InWorld->GetGameInstance() != nullptr )
+	if (InWorld != nullptr && InWorld->GetGameInstance() != nullptr)
 	{
 		FString DemoName;
 
@@ -2902,19 +2902,19 @@ bool UWorld::HandleDemoPlayCommand( const TCHAR* Cmd, FOutputDevice& Ar, UWorld*
 	{
 		ErrorString = TEXT( "InWorld->GetGameInstance() is null" );
 	}
-
-	if ( ErrorString != nullptr )
+	
+	if (ErrorString != nullptr)
 	{
-		Ar.Log( ErrorString );
+		Ar.Log(ErrorString);
 
-		if ( GetGameInstance() != nullptr )
+		if (GetGameInstance() != nullptr)
 		{
-			GetGameInstance()->HandleDemoPlaybackFailure( EDemoPlayFailure::Generic, FString( ErrorString ) );
+			GetGameInstance()->HandleDemoPlaybackFailure(EDemoPlayFailure::Generic, FString(ErrorString));
 		}
 	}
 	else
 	{
-		InWorld->GetGameInstance()->PlayReplay( Temp );
+		InWorld->GetGameInstance()->PlayReplay(Temp);
 	}
 
 	return true;
@@ -4534,7 +4534,7 @@ bool FSeamlessTravelHandler::StartTravel(UWorld* InCurrentWorld, const FURL& InU
 			FName CurrentMapName = CurrentWorld->GetOutermost()->GetFName();
 			FName DestinationMapName = FName(*PendingTravelURL.Map);
 
-			FString TransitionMap = GetDefault<UGameMapsSettings>()->TransitionMap;
+			FString TransitionMap = GetDefault<UGameMapsSettings>()->TransitionMap.GetLongPackageName();
 			FName DefaultMapFinalName(*TransitionMap);
 
 			// if we're already in the default map, skip loading it and just go to the destination
diff --git a/Engine/Source/Runtime/Engine/Public/Animation/AnimInstanceProxy.h b/Engine/Source/Runtime/Engine/Public/Animation/AnimInstanceProxy.h
--- a/Engine/Source/Runtime/Engine/Public/Animation/AnimInstanceProxy.h
+++ b/Engine/Source/Runtime/Engine/Public/Animation/AnimInstanceProxy.h
@@ -338,11 +338,8 @@ protected:
 	 */
 	void RecalcRequiredBones(USkeletalMeshComponent* Component, UObject* Asset);
 
-	/** Update morph target curves for external map (from USkeletalMeshComponent) */
-	void UpdateMorphTargetCurves(const TMap<FName, float>& InMorphTargetCurves);
-
 	/** Update the material parameters of the supplied component from this instance */
-	void UpdateComponentsMaterialParameters(UPrimitiveComponent* Component);
+	void UpdateCurvesToComponents(USkeletalMeshComponent* Component=nullptr);
 
 	/** Get Currently active montage evaluation state.
 		Note that there might be multiple Active at the same time. This will only return the first active one it finds. **/
diff --git a/Engine/Source/Runtime/Engine/Public/AnimationRuntime.h b/Engine/Source/Runtime/Engine/Public/AnimationRuntime.h
--- a/Engine/Source/Runtime/Engine/Public/AnimationRuntime.h
+++ b/Engine/Source/Runtime/Engine/Public/AnimationRuntime.h
@@ -306,7 +306,7 @@ public:
 	/**
 	* Combine CurveKeys (that reference morph targets by name) and ActiveAnims (that reference vertex anims by reference) into the ActiveVertexAnims array.
 	*/
-	static TArray<struct FActiveVertexAnim> UpdateActiveVertexAnims(const USkeletalMesh* InSkeletalMesh, const TMap<FName, float>& InMorphCurveAnims, const TArray<FActiveVertexAnim>& InActiveAnims);
+ 	static void AppendActiveVertexAnims(const USkeletalMesh* InSkeletalMesh, const TMap<FName, float>& InMorphCurveAnims, TArray<FActiveVertexAnim>& InOutActiveAnims);
 
 private:
 	/** 
diff --git a/Engine/Source/Runtime/Engine/Public/UnrealClient.h b/Engine/Source/Runtime/Engine/Public/UnrealClient.h
--- a/Engine/Source/Runtime/Engine/Public/UnrealClient.h
+++ b/Engine/Source/Runtime/Engine/Public/UnrealClient.h
@@ -632,21 +632,6 @@ extern ENGINE_API bool IsAltDown(FViewport* Viewport);
 extern ENGINE_API bool GetViewportScreenShot(FViewport* Viewport, TArray<FColor>& Bitmap, const FIntRect& ViewRect = FIntRect());
 extern ENGINE_API bool GetHighResScreenShotInput(const TCHAR* Cmd, FOutputDevice& Ar, uint32& OutXRes, uint32& OutYRes, float& OutResMult, FIntRect& OutCaptureRegion, bool& OutShouldEnableMask);
 
-namespace EMouseCaptureMode
-{
-	enum Type
-	{
-		/** Do not capture the mouse at all */
-		NoCapture,
-		/** Capture the mouse permanently when the viewport is clicked */
-		CapturePermanently,
-		/** Capture the mouse during a mouse down, releases on mouse up */
-		CaptureDuringMouseDown,
-		/** Capture only when the right mouse button is down, not any of the other mouse buttons */
-		CaptureDuringRightMouseDown,
-	};
-}
-
 /**
  * An abstract interface to a viewport's client.
  * The viewport's client processes input received by the viewport, and draws the viewport.
@@ -904,7 +889,7 @@ public:
 	/**
 	 * Gets the mouse capture behavior when the viewport is clicked
 	 */
-	virtual EMouseCaptureMode::Type CaptureMouseOnClick() { return EMouseCaptureMode::CapturePermanently; }
+	virtual EMouseCaptureMode CaptureMouseOnClick() { return EMouseCaptureMode::CapturePermanently; }
 
 	/**
 	 * Gets whether or not the cursor is hidden when the viewport captures the mouse
diff --git a/Engine/Source/Runtime/EngineSettings/Classes/GameMapsSettings.h b/Engine/Source/Runtime/EngineSettings/Classes/GameMapsSettings.h
--- a/Engine/Source/Runtime/EngineSettings/Classes/GameMapsSettings.h
+++ b/Engine/Source/Runtime/EngineSettings/Classes/GameMapsSettings.h
@@ -41,7 +41,7 @@ class ENGINESETTINGS_API UGameMapsSettings
 	 *
 	 * @return the default map specified in the settings
 	 */
-	static const FString& GetGameDefaultMap( );
+	static const FString GetGameDefaultMap( );
 
 	/**
 	 * Get the global default game type specified in the configuration
@@ -65,19 +65,22 @@ class ENGINESETTINGS_API UGameMapsSettings
 	 */
 	static void SetGlobalDefaultGameMode( const FString& NewGameMode );
 
+	virtual void PostInitProperties() override;
+	virtual void PostReloadConfig( class UProperty* PropertyThatWasLoaded ) override;
+
 public:
 
 	/** If set, this map will be loaded when the Editor starts up. */
-	UPROPERTY(config, EditAnywhere, Category=DefaultMaps)
-	FString EditorStartupMap;
+	UPROPERTY(config, EditAnywhere, Category=DefaultMaps, meta=(AllowedClasses="World"))
+	FStringAssetReference EditorStartupMap;
 
 	/** The default options that will be appended to a map being loaded. */
 	UPROPERTY(config, EditAnywhere, Category=DefaultMaps, AdvancedDisplay)
 	FString LocalMapOptions;
 
 	/** The map loaded when transition from one map to another. */
-	UPROPERTY(config, EditAnywhere, Category=DefaultMaps, AdvancedDisplay)
-	FString TransitionMap;
+	UPROPERTY(config, EditAnywhere, Category=DefaultMaps, AdvancedDisplay, meta=(AllowedClasses="World"))
+	FStringAssetReference TransitionMap;
 
 	/** Whether the screen should be split or not when multiple local players are present */
 	UPROPERTY(config, EditAnywhere, Category=LocalMultiplayer)
@@ -98,12 +101,12 @@ public:
 private:
 
 	/** The map that will be loaded by default when no other map is loaded. */
-	UPROPERTY(config, EditAnywhere, Category=DefaultMaps)
-	FString GameDefaultMap;
+	UPROPERTY(config, EditAnywhere, Category=DefaultMaps, meta=(AllowedClasses="World"))
+	FStringAssetReference GameDefaultMap;
 
 	/** The map that will be loaded by default when no other map is loaded (DEDICATED SERVER). */
-	UPROPERTY(config, EditAnywhere, Category=DefaultMaps, AdvancedDisplay)
-	FString ServerDefaultMap;
+	UPROPERTY(config, EditAnywhere, Category=DefaultMaps, AdvancedDisplay, meta=(AllowedClasses="World"))
+	FStringAssetReference ServerDefaultMap;
 
 	/** GameMode to use if not specified in any other way. (e.g. per-map DefaultGameMode or on the URL). */
 	UPROPERTY(config, noclear, EditAnywhere, Category=DefaultModes, meta=(MetaClass="GameMode", DisplayName="Default GameMode"))
diff --git a/Engine/Source/Runtime/EngineSettings/Private/EngineSettingsModule.cpp b/Engine/Source/Runtime/EngineSettings/Private/EngineSettingsModule.cpp
--- a/Engine/Source/Runtime/EngineSettings/Private/EngineSettingsModule.cpp
+++ b/Engine/Source/Runtime/EngineSettings/Private/EngineSettingsModule.cpp
@@ -67,11 +67,11 @@ UHudSettings::UHudSettings( const FObjectInitializer& ObjectInitializer )
 /* Static functions
  *****************************************************************************/
 
-const FString& UGameMapsSettings::GetGameDefaultMap( )
+const FString UGameMapsSettings::GetGameDefaultMap( )
 {
 	return IsRunningDedicatedServer()
-		? GetDefault<UGameMapsSettings>()->ServerDefaultMap
-		: GetDefault<UGameMapsSettings>()->GameDefaultMap;
+		? GetDefault<UGameMapsSettings>()->ServerDefaultMap.GetLongPackageName()
+		: GetDefault<UGameMapsSettings>()->GameDefaultMap.GetLongPackageName();
 }
 
 
@@ -106,5 +106,58 @@ void UGameMapsSettings::SetGlobalDefaultGameMode( const FString& NewGameMode )
 	GameMapsSettings->GlobalDefaultGameMode = NewGameMode;
 }
 
+// Backwards compat for map strings
+void FixMapAssetRef(FStringAssetReference& MapAssetReference)
+{
+	const FString AssetRefStr = MapAssetReference.ToString();
+	int32 DummyIndex;
+	if (!AssetRefStr.IsEmpty() && !AssetRefStr.FindLastChar(TEXT('.'), DummyIndex))
+	{
+		FString MapName, MapPath;
+		AssetRefStr.Split(TEXT("/"), &MapPath, &MapName, ESearchCase::IgnoreCase, ESearchDir::FromEnd);
+		MapAssetReference.SetPath(FString::Printf(TEXT("%s/%s.%s"),*MapPath, *MapName, *MapName));
+	}
+};
+
+void UGameMapsSettings::PostInitProperties()
+{
+	Super::PostInitProperties();
+
+	FixMapAssetRef(EditorStartupMap);
+	FixMapAssetRef(GameDefaultMap);
+	FixMapAssetRef(ServerDefaultMap);
+	FixMapAssetRef(TransitionMap);
+}
+
+void UGameMapsSettings::PostReloadConfig( UProperty* PropertyThatWasLoaded )
+{
+	if (PropertyThatWasLoaded)
+	{
+		if (PropertyThatWasLoaded->GetFName() == GET_MEMBER_NAME_CHECKED(UGameMapsSettings, EditorStartupMap))
+		{
+			FixMapAssetRef(EditorStartupMap);
+		}
+		else if (PropertyThatWasLoaded->GetFName() == GET_MEMBER_NAME_CHECKED(UGameMapsSettings, GameDefaultMap))
+		{
+			FixMapAssetRef(GameDefaultMap);
+		}
+		else if (PropertyThatWasLoaded->GetFName() == GET_MEMBER_NAME_CHECKED(UGameMapsSettings, ServerDefaultMap))
+		{
+			FixMapAssetRef(ServerDefaultMap);
+		}
+		else if (PropertyThatWasLoaded->GetFName() == GET_MEMBER_NAME_CHECKED(UGameMapsSettings, TransitionMap))
+		{
+			FixMapAssetRef(TransitionMap);
+		}
+	}
+	else
+	{
+		FixMapAssetRef(EditorStartupMap);
+		FixMapAssetRef(GameDefaultMap);
+		FixMapAssetRef(ServerDefaultMap);
+		FixMapAssetRef(TransitionMap);
+	}
+}
+
 
 IMPLEMENT_MODULE(FEngineSettingsModule, EngineSettings);
diff --git a/Engine/Source/Runtime/GameplayAbilities/Public/GameplayAbilitiesModule.h b/Engine/Source/Runtime/GameplayAbilities/Public/GameplayAbilitiesModule.h
--- a/Engine/Source/Runtime/GameplayAbilities/Public/GameplayAbilitiesModule.h
+++ b/Engine/Source/Runtime/GameplayAbilities/Public/GameplayAbilitiesModule.h
@@ -3,7 +3,6 @@
 #pragma once
 
 #include "ModuleManager.h"
-#include "AbilitySystemGlobals.h"
 
 
 /**
@@ -38,7 +37,7 @@ public:
 		return FModuleManager::Get().IsModuleLoaded( "GameplayAbilities" );
 	}
 
-	virtual UAbilitySystemGlobals* GetAbilitySystemGlobals() = 0;
+	virtual class UAbilitySystemGlobals* GetAbilitySystemGlobals() = 0;
 
 	virtual bool IsAbilitySystemGlobalsAvailable() = 0;
 };
diff --git a/Engine/Source/Runtime/HeadMountedDisplay/Public/IHeadMountedDisplay.h b/Engine/Source/Runtime/HeadMountedDisplay/Public/IHeadMountedDisplay.h
--- a/Engine/Source/Runtime/HeadMountedDisplay/Public/IHeadMountedDisplay.h
+++ b/Engine/Source/Runtime/HeadMountedDisplay/Public/IHeadMountedDisplay.h
@@ -334,6 +334,9 @@ public:
 	virtual FVector2D GetTextureOffsetRight() const {return FVector2D::ZeroVector;}
 	virtual FVector2D GetTextureScaleLeft() const {return FVector2D::ZeroVector;}
 	virtual FVector2D GetTextureScaleRight() const {return FVector2D::ZeroVector;}
+	virtual const float* GetRedDistortionParameters() const { return nullptr; }
+	virtual const float* GetGreenDistortionParameters() const { return nullptr; }
+	virtual const float* GetBlueDistortionParameters() const { return nullptr; }
 
 	virtual bool NeedsUpscalePostProcessPass()  { return false; }
 
diff --git a/Engine/Source/Runtime/InputCore/Private/InputCoreTypes.cpp b/Engine/Source/Runtime/InputCore/Private/InputCoreTypes.cpp
--- a/Engine/Source/Runtime/InputCore/Private/InputCoreTypes.cpp
+++ b/Engine/Source/Runtime/InputCore/Private/InputCoreTypes.cpp
@@ -566,16 +566,18 @@ void EKeys::Initialize()
 	AddKey(FKeyDetails(EKeys::Acceleration, LOCTEXT("Acceleration", "Acceleration"), FKeyDetails::VectorAxis, "Motion"));
 
 	// Fingers
-	AddKey(FKeyDetails(EKeys::TouchKeys[ETouchIndex::Touch1], LOCTEXT("Touch1", "Touch 1"), FKeyDetails::NotBlueprintBindableKey));
-	AddKey(FKeyDetails(EKeys::TouchKeys[ETouchIndex::Touch2], LOCTEXT("Touch2", "Touch 2"), FKeyDetails::NotBlueprintBindableKey));
-	AddKey(FKeyDetails(EKeys::TouchKeys[ETouchIndex::Touch3], LOCTEXT("Touch3", "Touch 3"), FKeyDetails::NotBlueprintBindableKey));
-	AddKey(FKeyDetails(EKeys::TouchKeys[ETouchIndex::Touch4], LOCTEXT("Touch4", "Touch 4"), FKeyDetails::NotBlueprintBindableKey));
-	AddKey(FKeyDetails(EKeys::TouchKeys[ETouchIndex::Touch5], LOCTEXT("Touch5", "Touch 5"), FKeyDetails::NotBlueprintBindableKey));
-	AddKey(FKeyDetails(EKeys::TouchKeys[ETouchIndex::Touch6], LOCTEXT("Touch6", "Touch 6"), FKeyDetails::NotBlueprintBindableKey));
-	AddKey(FKeyDetails(EKeys::TouchKeys[ETouchIndex::Touch7], LOCTEXT("Touch7", "Touch 7"), FKeyDetails::NotBlueprintBindableKey));
-	AddKey(FKeyDetails(EKeys::TouchKeys[ETouchIndex::Touch8], LOCTEXT("Touch8", "Touch 8"), FKeyDetails::NotBlueprintBindableKey));
-	AddKey(FKeyDetails(EKeys::TouchKeys[ETouchIndex::Touch9], LOCTEXT("Touch9", "Touch 9"), FKeyDetails::NotBlueprintBindableKey));
-	AddKey(FKeyDetails(EKeys::TouchKeys[ETouchIndex::Touch10], LOCTEXT("Touch10", "Touch 10"), FKeyDetails::NotBlueprintBindableKey));
+	AddMenuCategoryDisplayInfo("Touch", LOCTEXT("TouchSubCateogry", "Touch"), TEXT("GraphEditor.TouchEvent_16x"));
+
+	AddKey(FKeyDetails(EKeys::TouchKeys[ETouchIndex::Touch1], LOCTEXT("Touch1", "Touch 1"), 0, "Touch"));
+	AddKey(FKeyDetails(EKeys::TouchKeys[ETouchIndex::Touch2], LOCTEXT("Touch2", "Touch 2"), 0, "Touch"));
+	AddKey(FKeyDetails(EKeys::TouchKeys[ETouchIndex::Touch3], LOCTEXT("Touch3", "Touch 3"), 0, "Touch"));
+	AddKey(FKeyDetails(EKeys::TouchKeys[ETouchIndex::Touch4], LOCTEXT("Touch4", "Touch 4"), 0, "Touch"));
+	AddKey(FKeyDetails(EKeys::TouchKeys[ETouchIndex::Touch5], LOCTEXT("Touch5", "Touch 5"), 0, "Touch"));
+	AddKey(FKeyDetails(EKeys::TouchKeys[ETouchIndex::Touch6], LOCTEXT("Touch6", "Touch 6"), 0, "Touch"));
+	AddKey(FKeyDetails(EKeys::TouchKeys[ETouchIndex::Touch7], LOCTEXT("Touch7", "Touch 7"), 0, "Touch"));
+	AddKey(FKeyDetails(EKeys::TouchKeys[ETouchIndex::Touch8], LOCTEXT("Touch8", "Touch 8"), 0, "Touch"));
+	AddKey(FKeyDetails(EKeys::TouchKeys[ETouchIndex::Touch9], LOCTEXT("Touch9", "Touch 9"), 0, "Touch"));
+	AddKey(FKeyDetails(EKeys::TouchKeys[ETouchIndex::Touch10], LOCTEXT("Touch10", "Touch 10"), 0, "Touch"));
 
 	// Gestures
 	AddMenuCategoryDisplayInfo("Gesture", LOCTEXT("GestureSubCateogry", "Gesture"), TEXT("GraphEditor.KeyEvent_16x"));
@@ -648,7 +650,7 @@ void EKeys::Initialize()
 
 	// PS4-specific
 	AddMenuCategoryDisplayInfo("PS4", LOCTEXT("PS4SubCategory", "PS4"), TEXT("GraphEditor.PadEvent_16x"));
-	AddKey(FKeyDetails(EKeys::PS4_Special, LOCTEXT("PS4_Special", "PS4_Special"), FKeyDetails::NotBlueprintBindableKey));
+	AddKey(FKeyDetails(EKeys::PS4_Special, LOCTEXT("PS4_Special", "PS4_Special"), FKeyDetails::NotBlueprintBindableKey, "PS4"));
 	AddKey(FKeyDetails(EKeys::Gamepad_Special_Left_X, LOCTEXT("PS4_Gamepad_Special_Left_X", "PS4 Touchpad Button X Axis"), FKeyDetails::GamepadKey | FKeyDetails::FloatAxis, "PS4"));
 	AddKey(FKeyDetails(EKeys::Gamepad_Special_Left_Y, LOCTEXT("PS4_Gamepad_Special_Left_Y", "PS4 Touchpad Button Y Axis"), FKeyDetails::GamepadKey | FKeyDetails::FloatAxis, "PS4"));
 
diff --git a/Engine/Source/Runtime/Launch/Private/Android/AndroidJNI.cpp b/Engine/Source/Runtime/Launch/Private/Android/AndroidJNI.cpp
--- a/Engine/Source/Runtime/Launch/Private/Android/AndroidJNI.cpp
+++ b/Engine/Source/Runtime/Launch/Private/Android/AndroidJNI.cpp
@@ -62,6 +62,9 @@ void FJavaWrapper::FindClassesAndMethods(JNIEnv* Env)
 	AndroidThunkJava_GetMetaDataInt = FindMethod(Env, GameActivityClassID, "AndroidThunkJava_GetMetaDataInt", "(Ljava/lang/String;)I", bIsOptional);
 	AndroidThunkJava_GetMetaDataString = FindMethod(Env, GameActivityClassID, "AndroidThunkJava_GetMetaDataString", "(Ljava/lang/String;)Ljava/lang/String;", bIsOptional);
 
+	// this is optional - only inserted if GearVR plugin enabled
+	AndroidThunkJava_IsGearVRApplication = FindMethod(Env, GameActivityClassID, "AndroidThunkJava_IsGearVRApplication", "()Z", true);
+
 	// get field IDs for InputDeviceInfo class members
 	jclass localInputDeviceInfoClass = FindClass(Env, "com/epicgames/ue4/GameActivity$InputDeviceInfo", bIsOptional);
 	InputDeviceInfoClass = (jclass)Env->NewGlobalRef(localInputDeviceInfoClass);
@@ -204,6 +207,7 @@ jmethodID FJavaWrapper::AndroidThunkJava_HasMetaDataKey;
 jmethodID FJavaWrapper::AndroidThunkJava_GetMetaDataBoolean;
 jmethodID FJavaWrapper::AndroidThunkJava_GetMetaDataInt;
 jmethodID FJavaWrapper::AndroidThunkJava_GetMetaDataString;
+jmethodID FJavaWrapper::AndroidThunkJava_IsGearVRApplication;
 
 jclass FJavaWrapper::InputDeviceInfoClass;
 jfieldID FJavaWrapper::InputDeviceInfo_VendorId;
@@ -389,6 +393,25 @@ FString AndroidThunkCpp_GetMetaDataString(const FString& Key)
 	return Result;
 }
 
+// call out to JNI to see if the application was packaged for GearVR
+bool AndroidThunkCpp_IsGearVRApplication()
+{
+	static int32 IsGearVRApplication = -1;
+
+	if (IsGearVRApplication == -1)
+	{
+		IsGearVRApplication = 0;
+		if (FJavaWrapper::AndroidThunkJava_IsGearVRApplication)
+		{
+			if (JNIEnv* Env = FAndroidApplication::GetJavaEnv())
+			{
+				IsGearVRApplication = FJavaWrapper::CallBooleanMethod(Env, FJavaWrapper::GameActivityThis, FJavaWrapper::AndroidThunkJava_IsGearVRApplication) ? 1 : 0;
+			}
+		}
+	}
+	return IsGearVRApplication == 1;
+}
+
 void AndroidThunkCpp_ShowConsoleWindow()
 {
 	if (JNIEnv* Env = FAndroidApplication::GetJavaEnv())
diff --git a/Engine/Source/Runtime/Launch/Private/Android/LaunchAndroid.cpp b/Engine/Source/Runtime/Launch/Private/Android/LaunchAndroid.cpp
--- a/Engine/Source/Runtime/Launch/Private/Android/LaunchAndroid.cpp
+++ b/Engine/Source/Runtime/Launch/Private/Android/LaunchAndroid.cpp
@@ -304,9 +304,11 @@ int32 AndroidMain(struct android_app* state)
 		if(!FAppEventManager::GetInstance()->IsGamePaused())
 		{
 			GEngineLoop.Tick();
-
-			float timeToSleep = 0.05f; //in seconds
-			sleep(timeToSleep);
+		}
+		else
+		{
+			// use less CPU when paused
+			FPlatformProcess::Sleep(0.10f);
 		}
 
 #if !UE_BUILD_SHIPPING
diff --git a/Engine/Source/Runtime/Networking/Public/Common/UdpSocketReceiver.h b/Engine/Source/Runtime/Networking/Public/Common/UdpSocketReceiver.h
--- a/Engine/Source/Runtime/Networking/Public/Common/UdpSocketReceiver.h
+++ b/Engine/Source/Runtime/Networking/Public/Common/UdpSocketReceiver.h
@@ -30,18 +30,19 @@ public:
 	 *
 	 * @param InSocket The UDP socket to receive data from.
 	 * @param InWaitTime The amount of time to wait for the socket to be readable.
-	 * @param ThreadDescription The thread description text (for debugging).
+	 * @param InThreadName The receiver thread name (for debugging).
 	 */
-	FUdpSocketReceiver(FSocket* InSocket, const FTimespan& InWaitTime, const TCHAR* ThreadDescription)
+	FUdpSocketReceiver(FSocket* InSocket, const FTimespan& InWaitTime, const TCHAR* InThreadName)
 		: Socket(InSocket)
 		, Stopping(false)
+		, Thread(nullptr)
+		, ThreadName(InThreadName)
 		, WaitTime(InWaitTime)
 	{
 		check(Socket != nullptr);
 		check(Socket->GetSocketType() == SOCKTYPE_Datagram);
 
 		SocketSubsystem = ISocketSubsystem::Get(PLATFORM_SOCKETSUBSYSTEM);
-		Thread = FRunnableThread::Create(this, ThreadDescription, 128 * 1024, TPri_AboveNormal, FPlatformAffinity::GetPoolThreadMask());
 	}
 
 	/** Virtual destructor. */
@@ -56,13 +57,23 @@ public:
 
 public:
 
+	/** Start the receiver thread. */
+	void Start()
+	{
+		Thread = FRunnableThread::Create(this, *ThreadName, 128 * 1024, TPri_AboveNormal, FPlatformAffinity::GetPoolThreadMask());
+	}
+
 	/**
 	 * Returns a delegate that is executed when data has been received.
 	 *
+	 * This delegate must be bound before the receiver thread is started with
+	 * the Start() method. It cannot be unbound while the thread is running.
+	 *
 	 * @return The delegate.
 	 */
 	FOnSocketDataReceived& OnDataReceived()
 	{
+		check(Thread == nullptr);
 		return DataReceivedDelegate;
 	}
 
@@ -127,6 +138,9 @@ private:
 	/** Holds the thread object. */
 	FRunnableThread* Thread;
 
+	/** The receiver thread's name. */
+	FString ThreadName;
+
 	/** Holds the amount of time to wait for inbound packets. */
 	FTimespan WaitTime;
 
diff --git a/Engine/Source/Runtime/Online/IOS/OnlineSubsystemIOS/Private/OnlineStoreInterfaceIOS.cpp b/Engine/Source/Runtime/Online/IOS/OnlineSubsystemIOS/Private/OnlineStoreInterfaceIOS.cpp
--- a/Engine/Source/Runtime/Online/IOS/OnlineSubsystemIOS/Private/OnlineStoreInterfaceIOS.cpp
+++ b/Engine/Source/Runtime/Online/IOS/OnlineSubsystemIOS/Private/OnlineStoreInterfaceIOS.cpp
@@ -277,6 +277,17 @@
 #endif
 }
 
+-(void)request:(SKRequest*)request didFailWithError: (NSError *)error
+{
+#ifdef __IPHONE_7_0
+	if ([Request isKindOfClass : [SKReceiptRefreshRequest class]])
+	{
+		[self paymentQueue:[SKPaymentQueue defaultQueue]  restoreCompletedTransactionsFailedWithError: error];
+		[Request release];
+	}
+#endif
+}
+
 -(void)restorePurchases
 {
 #ifdef __IPHONE_7_0
diff --git a/Engine/Source/Runtime/OpenGLDrv/Private/Android/AndroidEGL.cpp b/Engine/Source/Runtime/OpenGLDrv/Private/Android/AndroidEGL.cpp
--- a/Engine/Source/Runtime/OpenGLDrv/Private/Android/AndroidEGL.cpp
+++ b/Engine/Source/Runtime/OpenGLDrv/Private/Android/AndroidEGL.cpp
@@ -224,6 +224,7 @@ void AndroidEGL::CreateEGLSurface(ANativeWindow* InWindow, bool bCreateWndSurfac
 	// due to possible early initialization, don't redo this
 	if (PImplData->eglSurface != EGL_NO_SURFACE)
 	{
+		FPlatformMisc::LowLevelOutputDebugStringf(TEXT("AndroidEGL::CreateEGLSurface() Already initialized: %p"), PImplData->eglSurface);
 		return;
 	}
 
@@ -512,9 +513,9 @@ void AndroidEGL::InitBackBuffer()
 
 extern void AndroidThunkCpp_SetDesiredViewSize(int32 Width, int32 Height);
 
-void AndroidEGL::InitSurface(bool bUseSmallSurface)
+void AndroidEGL::InitSurface(bool bUseSmallSurface, bool bCreateWndSurface)
 {
-	FPlatformMisc::LowLevelOutputDebugStringf( TEXT("AndroidEGL::InitSurface %d" ), int(bUseSmallSurface));
+	FPlatformMisc::LowLevelOutputDebugStringf(TEXT("AndroidEGL::InitSurface %d, %d"), int(bUseSmallSurface), int(bCreateWndSurface));
 	ANativeWindow* window = (ANativeWindow*)FPlatformMisc::GetHardwareWindow();
 	while(window == NULL)
 	{
@@ -533,7 +534,7 @@ void AndroidEGL::InitSurface(bool bUseSmallSurface)
 		AndroidThunkCpp_SetDesiredViewSize(Width, Height);
 	}
 	ANativeWindow_setBuffersGeometry(PImplData->Window, Width, Height, PImplData->NativeVisualID);
-	CreateEGLSurface(PImplData->Window, bUseSmallSurface);
+	CreateEGLSurface(PImplData->Window, bCreateWndSurface);
 	
 	PImplData->SharedContext.eglSurface = PImplData->auxSurface;
 	PImplData->RenderingContext.eglSurface = PImplData->eglSurface;
@@ -543,8 +544,9 @@ void AndroidEGL::InitSurface(bool bUseSmallSurface)
 
 void AndroidEGL::ReInit()
 {
+	FPlatformMisc::LowLevelOutputDebugString(TEXT("AndroidEGL::ReInit()"));
 	SetCurrentContext(EGL_NO_CONTEXT, EGL_NO_SURFACE);
-	InitSurface(false);
+	InitSurface(false, true);
 	SetCurrentSharedContext();
 }
 
@@ -806,17 +808,20 @@ FPlatformOpenGLContext* AndroidEGL::GetRenderingContext()
 
 void AndroidEGL::UnBind()
 {
+	FPlatformMisc::LowLevelOutputDebugString(TEXT("AndroidEGL::UnBind()"));
 	ResetDisplay();
 	DestroySurface();
 }
 
 void FAndroidAppEntry::ReInitWindow()
 {
+	FPlatformMisc::LowLevelOutputDebugString(TEXT("AndroidEGL::ReInitWindow()"));
 	AndroidEGL::GetInstance()->ReInit();
 }
 
 void FAndroidAppEntry::DestroyWindow()
 {
+	FPlatformMisc::LowLevelOutputDebugString(TEXT("AndroidEGL::DestroyWindow()"));
 	AndroidEGL::GetInstance()->UnBind();
 }
 
diff --git a/Engine/Source/Runtime/OpenGLDrv/Private/Android/AndroidEGL.h b/Engine/Source/Runtime/OpenGLDrv/Private/Android/AndroidEGL.h
--- a/Engine/Source/Runtime/OpenGLDrv/Private/Android/AndroidEGL.h
+++ b/Engine/Source/Runtime/OpenGLDrv/Private/Android/AndroidEGL.h
@@ -65,7 +65,7 @@ public:
 	void UnBind();
 	bool SwapBuffers();
 	void Terminate();
-	void InitSurface(bool bUseSmallSurface);
+	void InitSurface(bool bUseSmallSurface, bool bCreateWndSurface);
 
 	void GetDimensions(uint32& OutWidth, uint32& OutHeight);
 	
diff --git a/Engine/Source/Runtime/OpenGLDrv/Private/Android/AndroidES31OpenGL.cpp b/Engine/Source/Runtime/OpenGLDrv/Private/Android/AndroidES31OpenGL.cpp
--- a/Engine/Source/Runtime/OpenGLDrv/Private/Android/AndroidES31OpenGL.cpp
+++ b/Engine/Source/Runtime/OpenGLDrv/Private/Android/AndroidES31OpenGL.cpp
@@ -123,7 +123,7 @@ void FPlatformOpenGLDevice::Init()
 	UE_LOG( LogRHI, Warning, TEXT("Entering FPlatformOpenGLDevice::Init"));
 	extern void InitDebugContext();
 
-	AndroidEGL::GetInstance()->InitSurface(false);
+	AndroidEGL::GetInstance()->InitSurface(false, true);
 	AndroidEGL::GetInstance()->SetSingleThreadRenderingContext();
 
 	// Initialize all of the entry points we have to query manually
diff --git a/Engine/Source/Runtime/OpenGLDrv/Private/Android/AndroidOpenGL.cpp b/Engine/Source/Runtime/OpenGLDrv/Private/Android/AndroidOpenGL.cpp
--- a/Engine/Source/Runtime/OpenGLDrv/Private/Android/AndroidOpenGL.cpp
+++ b/Engine/Source/Runtime/OpenGLDrv/Private/Android/AndroidOpenGL.cpp
@@ -58,6 +58,14 @@ PFNGLDRAWELEMENTSINSTANCEDPROC			glDrawElementsInstanced = NULL;
 PFNGLDRAWARRAYSINSTANCEDPROC			glDrawArraysInstanced = NULL;
 PFNGLVERTEXATTRIBDIVISORPROC			glVertexAttribDivisor = NULL;
 
+static TAutoConsoleVariable<int32> CVarAndroidDisableTextureFormatBGRA8888(
+	TEXT("android.DisableTextureFormatBGRA8888"),
+	0,
+	TEXT("Whether to disable usage of GL_EXT_texture_format_BGRA8888 extension.\n")
+	TEXT(" 0: Enable when extension is available (default)\n")
+	TEXT(" 1: Always disabled"),
+	ECVF_ReadOnly);
+
 struct FPlatformOpenGLDevice
 {
 
@@ -88,7 +96,7 @@ void FPlatformOpenGLDevice::Init()
 {
 	extern void InitDebugContext();
 
-	AndroidEGL::GetInstance()->InitSurface(false);
+	AndroidEGL::GetInstance()->InitSurface(false, true);
 	PlatformRenderingContextSetup(this);
 
 	LoadEXT();
@@ -441,6 +449,12 @@ void FAndroidOpenGL::ProcessExtensions(const FString& ExtensionsString)
 		UE_LOG(LogRHI, Warning, TEXT("Disabling support for hardware instancing on Adreno 330 OpenGL ES 3.0 [REDACTED_EMAIL] AU@  (CL@)"));
 		bSupportsInstancing = false;
 	}
+
+	if (bSupportsBGRA8888 && CVarAndroidDisableTextureFormatBGRA8888.GetValueOnAnyThread() == 1)
+	{
+		UE_LOG(LogRHI, Warning, TEXT("Disabling support for GL_EXT_texture_format_BGRA8888"));
+		bSupportsBGRA8888 = false;
+	}
 }
 
 
diff --git a/Engine/Source/Runtime/OpenGLDrv/Private/OpenGLCommands.cpp b/Engine/Source/Runtime/OpenGLDrv/Private/OpenGLCommands.cpp
--- a/Engine/Source/Runtime/OpenGLDrv/Private/OpenGLCommands.cpp
+++ b/Engine/Source/Runtime/OpenGLDrv/Private/OpenGLCommands.cpp
@@ -640,7 +640,9 @@ inline void FOpenGLDynamicRHI::ApplyTextureStage(FOpenGLContextState& ContextSta
 		FOpenGL::TexParameter(Target, GL_TEXTURE_MAG_FILTER, SamplerState->Data.MagFilter);
 		if( FOpenGL::SupportsTextureFilterAnisotropic() )
 		{
-			FOpenGL::TexParameter(Target, GL_TEXTURE_MAX_ANISOTROPY_EXT, SamplerState->Data.MaxAnisotropy);
+			// GL_EXT_texture_filter_anisotropic requires value to be at least 1
+			GLint MaxAnisotropy = FMath::Max(1, SamplerState->Data.MaxAnisotropy);
+			FOpenGL::TexParameter(Target, GL_TEXTURE_MAX_ANISOTROPY_EXT, MaxAnisotropy);
 		}
 
 		if( FOpenGL::SupportsTextureCompare() )
diff --git a/Engine/Source/Runtime/OpenGLDrv/Private/OpenGLDevice.cpp b/Engine/Source/Runtime/OpenGLDrv/Private/OpenGLDevice.cpp
--- a/Engine/Source/Runtime/OpenGLDrv/Private/OpenGLDevice.cpp
+++ b/Engine/Source/Runtime/OpenGLDrv/Private/OpenGLDevice.cpp
@@ -815,7 +815,7 @@ static void InitRHICapabilitiesForGL()
 		GLuint RGBA8 = FOpenGL::SupportsRGBA8() ? GL_RGBA8_OES : GL_RGBA;
 
 	#if PLATFORM_ANDROID
-		SetupTextureFormat(PF_B8G8R8A8, FOpenGLTextureFormat(GL_BGRA, FOpenGL::SupportsSRGB() ? GL_SRGB_ALPHA_EXT : GL_BGRA, BGRA8888, GL_UNSIGNED_BYTE, false, false));
+		SetupTextureFormat(PF_B8G8R8A8, FOpenGLTextureFormat(BGRA8888, FOpenGL::SupportsSRGB() ? GL_SRGB_ALPHA_EXT : BGRA8888, BGRA8888, GL_UNSIGNED_BYTE, false, false));
 	#else
 		SetupTextureFormat(PF_B8G8R8A8, FOpenGLTextureFormat(GL_RGBA, FOpenGL::SupportsSRGB() ? GL_SRGB_ALPHA_EXT : GL_RGBA, GL_BGRA8_EXT, FOpenGL::SupportsSRGB() ? GL_SRGB8_ALPHA8_EXT : GL_BGRA8_EXT, BGRA8888, GL_UNSIGNED_BYTE, false, false));
 	#endif
diff --git a/Engine/Source/Runtime/OpenGLDrv/Private/OpenGLShaders.cpp b/Engine/Source/Runtime/OpenGLDrv/Private/OpenGLShaders.cpp
--- a/Engine/Source/Runtime/OpenGLDrv/Private/OpenGLShaders.cpp
+++ b/Engine/Source/Runtime/OpenGLDrv/Private/OpenGLShaders.cpp
@@ -1150,7 +1150,7 @@ public:
 		for (int32 Index = 0; Index < PackedGlobalArrays.Num(); ++Index)
 		{
 			auto& PackedArray = PackedGlobalArrays[Index];
-			int32 FoundIndex = -1;
+			FPackedUniformInfo OutInfo = {-1, PackedArray.TypeName, CrossCompiler::PACKED_TYPEINDEX_MAX};
 
 			// Find this Global Array in the reflection list
 			for (int32 FindIndex = 0; FindIndex < ReflectedUniformInfos.Num(); ++FindIndex)
@@ -1158,11 +1158,12 @@ public:
 				auto& ReflectedInfo = ReflectedUniformInfos[FindIndex];
 				if (ReflectedInfo.ArrayType == PackedArray.TypeName)
 				{
-					FoundIndex = FindIndex;
-					OutPackedUniformInfos.Add(ReflectedInfo);
+					OutInfo = ReflectedInfo;
 					break;
 				}
 			}
+
+			OutPackedUniformInfos.Add(OutInfo);
 		}
 	}
 };
@@ -2729,6 +2730,12 @@ void FOpenGLShaderParameterCache::CommitPackedGlobals(const FOpenGLLinkedProgram
 	for (int32 PackedUniform = 0; PackedUniform < PackedUniforms.Num(); ++PackedUniform)
 	{
 		const FOpenGLLinkedProgram::FPackedUniformInfo& UniformInfo = PackedUniforms[PackedUniform];
+		if (UniformInfo.Location < 0)
+		{
+			// Probably this uniform array was optimized away in a linked program
+			continue;
+		}
+		
 		const uint32 ArrayIndex = UniformInfo.Index;
 		check(ArrayIndex < CrossCompiler::PACKED_TYPEINDEX_MAX);
 		const int32 NumVectors = PackedArrays[PackedUniform].Size / BytesPerRegister;
diff --git a/Engine/Source/Runtime/Projects/Private/ModuleDescriptor.cpp b/Engine/Source/Runtime/Projects/Private/ModuleDescriptor.cpp
--- a/Engine/Source/Runtime/Projects/Private/ModuleDescriptor.cpp
+++ b/Engine/Source/Runtime/Projects/Private/ModuleDescriptor.cpp
@@ -40,6 +40,9 @@ const TCHAR* ELoadingPhase::ToString( const ELoadingPhase::Type Value )
 	case PostEngineInit:
 		return TEXT( "PostEngineInit" );
 
+	case None:
+		return TEXT( "None" );
+
 	default:
 		ensureMsgf( false, TEXT( "Unrecognized ELoadingPhase value: %i" ), Value );
 		return NULL;
diff --git a/Engine/Source/Runtime/Projects/Public/ModuleDescriptor.h b/Engine/Source/Runtime/Projects/Public/ModuleDescriptor.h
--- a/Engine/Source/Runtime/Projects/Public/ModuleDescriptor.h
+++ b/Engine/Source/Runtime/Projects/Public/ModuleDescriptor.h
@@ -31,6 +31,9 @@ namespace ELoadingPhase
 		/** After the engine has been initialized */
 		PostEngineInit,
 
+		/** Do not automatically load this module */
+		None,
+
 		// NOTE: If you add a new value, make sure to update the ToString() method below!
 		Max
 	};
diff --git a/Engine/Source/Runtime/RHI/Private/RHI.cpp b/Engine/Source/Runtime/RHI/Private/RHI.cpp
--- a/Engine/Source/Runtime/RHI/Private/RHI.cpp
+++ b/Engine/Source/Runtime/RHI/Private/RHI.cpp
@@ -83,15 +83,13 @@ const FClearValueBinding FClearValueBinding::DepthFar((float)ERHIZBuffer::FarPla
 
 TLockFreePointerListUnordered<FRHIResource, PLATFORM_CACHE_LINE_SIZE> FRHIResource::PendingDeletes;
 FRHIResource* FRHIResource::CurrentlyDeleting = nullptr;
-TQueue<FRHIResource::ResourceToDelete> FRHIResource::DeferredDeletionQueue;
+TArray<FRHIResource::ResourcesToDelete> FRHIResource::DeferredDeletionQueue;
 uint32 FRHIResource::CurrentFrame = 0;
 
-#if !DISABLE_RHI_DEFFERED_DELETE
 bool FRHIResource::Bypass()
 {
 	return GRHICommandList.Bypass();
 }
-#endif
 
 DECLARE_CYCLE_STAT(TEXT("Delete Resources"), STAT_DeleteResources, STATGROUP_RHICMDLIST);
 
@@ -102,14 +100,9 @@ void FRHIResource::FlushPendingDeletes()
 	check(IsInRenderingThread());
 	FRHICommandListExecutor::CheckNoOutstandingCmdLists();
 	FRHICommandListExecutor::GetImmediateCommandList().ImmediateFlush(EImmediateFlushType::FlushRHIThread);
-	while (1)
+
+	auto Delete = [](TArray<FRHIResource*>& ToDelete)
 	{
-		TArray<FRHIResource*> ToDelete;
-		PendingDeletes.PopAll(ToDelete);
-		if (!ToDelete.Num())
-		{
-			break;
-		}
 		for (int32 Index = 0; Index < ToDelete.Num(); Index++)
 		{
 			FRHIResource* Ref = ToDelete[Index];
@@ -117,15 +110,7 @@ void FRHIResource::FlushPendingDeletes()
 			if (Ref->GetRefCount() == 0) // caches can bring dead objects back to life
 			{
 				CurrentlyDeleting = Ref;
-				if (PlatformNeedsExtraDeletionLatency())
-				{
-					DeferredDeletionQueue.Enqueue(ResourceToDelete(Ref, CurrentFrame));
-				}
-				else
-				{
-					delete Ref;
-				}
-
+				delete Ref;
 				CurrentlyDeleting = nullptr;
 			}
 			else
@@ -134,27 +119,43 @@ void FRHIResource::FlushPendingDeletes()
 				FPlatformMisc::MemoryBarrier();
 			}
 		}
+	};
+
+	while (1)
+	{
+		if (PendingDeletes.IsEmpty())
+		{
+			break;
+		}
+		if (PlatformNeedsExtraDeletionLatency())
+		{
+			const int32 Index = DeferredDeletionQueue.AddDefaulted();
+			ResourcesToDelete& ResourceBatch = DeferredDeletionQueue[Index];
+			ResourceBatch.FrameDeleted = CurrentFrame;
+			PendingDeletes.PopAll(ResourceBatch.Resources);
+			check(ResourceBatch.Resources.Num());
+		}
+		else
+		{
+			TArray<FRHIResource*> ToDelete;
+			PendingDeletes.PopAll(ToDelete);
+			check(ToDelete.Num());
+			Delete(ToDelete);
+		}
 	}
 
 	const uint32 NumFramesToExpire = 3;
 
-	if (PlatformNeedsExtraDeletionLatency())
+	if (DeferredDeletionQueue.Num())
 	{
-		ResourceToDelete TempResource;
-
-		while (!DeferredDeletionQueue.IsEmpty())
+		int32 DeletedBatchCount = 0;
+		while (DeletedBatchCount < DeferredDeletionQueue.Num())
 		{
-			DeferredDeletionQueue.Peek(TempResource);
-
-			if ((TempResource.FrameDeleted + NumFramesToExpire) < CurrentFrame)
+			ResourcesToDelete& ResourceBatch = DeferredDeletionQueue[DeletedBatchCount];
+			if (((ResourceBatch.FrameDeleted + NumFramesToExpire) < CurrentFrame) || !GIsRHIInitialized)
 			{
-				// It's possible the resource has been resurrected elsewhere
-				if (TempResource.Resource->GetRefCount() == 0)
-				{
-					delete TempResource.Resource;
-				}
-
-				DeferredDeletionQueue.Dequeue(TempResource);
+				Delete(ResourceBatch.Resources);
+				++DeletedBatchCount;
 			}
 			else
 			{
@@ -162,6 +163,11 @@ void FRHIResource::FlushPendingDeletes()
 			}
 		}
 
+		if (DeletedBatchCount)
+		{
+			DeferredDeletionQueue.RemoveAt(0, DeletedBatchCount);
+		}
+
 		++CurrentFrame;
 	}
 }
@@ -244,6 +250,7 @@ bool GSupportsSeparateRenderTargetBlendState = false;
 bool GSupportsDepthRenderTargetWithoutColorRenderTarget = true;
 float GMinClipZ = 0.0f;
 float GProjectionSignY = 1.0f;
+bool GRHINeedsExtraDeletionLatency = false;
 int32 GMaxShadowDepthBufferSizeX = 2048;
 int32 GMaxShadowDepthBufferSizeY = 2048;
 int32 GMaxTextureDimensions = 2048;
diff --git a/Engine/Source/Runtime/RHI/Public/RHI.h b/Engine/Source/Runtime/RHI/Public/RHI.h
--- a/Engine/Source/Runtime/RHI/Public/RHI.h
+++ b/Engine/Source/Runtime/RHI/Public/RHI.h
@@ -141,6 +141,9 @@ extern RHI_API float GMinClipZ;
 /** The sign to apply to the Y axis of projection matrices. */
 extern RHI_API float GProjectionSignY;
 
+/** Does this RHI need to wait for deletion of resources due to ref counting. */
+extern RHI_API bool GRHINeedsExtraDeletionLatency;
+
 /** The maximum size to allow for the shadow depth buffer in the X dimension.  This must be larger or equal to GMaxShadowDepthBufferSizeY. */
 extern RHI_API int32 GMaxShadowDepthBufferSizeX;
 /** The maximum size to allow for the shadow depth buffer in the Y dimension. */
diff --git a/Engine/Source/Runtime/RHI/Public/RHIResources.h b/Engine/Source/Runtime/RHI/Public/RHIResources.h
--- a/Engine/Source/Runtime/RHI/Public/RHIResources.h
+++ b/Engine/Source/Runtime/RHI/Public/RHIResources.h
@@ -22,10 +22,7 @@ public:
 	}
 	virtual ~FRHIResource() 
 	{
-		if (!PlatformNeedsExtraDeletionLatency())
-		{
-			check(NumRefs.GetValue() == 0 && (CurrentlyDeleting == this || bDoNotDeferDelete || Bypass())); // this should not have any outstanding refs
-		}
+		check(PlatformNeedsExtraDeletionLatency() || (NumRefs.GetValue() == 0 && (CurrentlyDeleting == this || bDoNotDeferDelete || Bypass()))); // this should not have any outstanding refs
 	}
 	FORCEINLINE_DEBUGGABLE uint32 AddRef() const
 	{
@@ -38,7 +35,7 @@ public:
 		int32 NewValue = NumRefs.Decrement();
 		if (NewValue == 0)
 		{
-			if (bDoNotDeferDelete || Bypass())
+			if (!DeferDelete())
 			{ 
 				delete this;
 			}
@@ -69,19 +66,12 @@ public:
 
 	static void FlushPendingDeletes();
 
-	static bool PlatformNeedsExtraDeletionLatency()
+	FORCEINLINE static bool PlatformNeedsExtraDeletionLatency()
 	{
-		return (PLATFORM_XBOXONE == 1);
+		return GRHINeedsExtraDeletionLatency && GIsRHIInitialized;
 	}
 
-#if DISABLE_RHI_DEFFERED_DELETE
-	FORCEINLINE static bool Bypass()
-	{
-		return true;
-	}
-#else
 	static bool Bypass();
-#endif
 
 private:
 	mutable FThreadSafeCounter NumRefs;
@@ -90,31 +80,31 @@ private:
 	static TLockFreePointerListUnordered<FRHIResource, PLATFORM_CACHE_LINE_SIZE> PendingDeletes;
 	static FRHIResource* CurrentlyDeleting;
 
+	FORCEINLINE bool DeferDelete() const
+	{
+#if DISABLE_RHI_DEFFERED_DELETE
+		return false;
+#else
+		// Defer if GRHINeedsExtraDeletionLatency or we are doing threaded rendering (unless otherwise requested).
+		return !bDoNotDeferDelete && (GRHINeedsExtraDeletionLatency || !Bypass());
+#endif
+	}
+
 	// Some APIs don't do internal reference counting, so we have to wait an extra couple of frames before deleting resources
 	// to ensure the GPU has completely finished with them. This avoids expensive fences, etc.
-	struct ResourceToDelete
+	struct ResourcesToDelete
 	{
-		ResourceToDelete()
-			: Resource(nullptr)
-			, FrameDeleted(0)
-		{
-
-		}
-
-		ResourceToDelete(
-			FRHIResource* InResource,
-			uint32 InFrameDeleted)
-			: Resource(InResource)
-			, FrameDeleted(InFrameDeleted)
+		ResourcesToDelete(uint32 InFrameDeleted = 0)
+			: FrameDeleted(InFrameDeleted)
 		{
 
 		}
 
-		FRHIResource*	Resource;
-		uint32			FrameDeleted;
+		TArray<FRHIResource*>	Resources;
+		uint32					FrameDeleted;
 	};
 
-	static TQueue<ResourceToDelete> DeferredDeletionQueue;
+	static TArray<ResourcesToDelete> DeferredDeletionQueue;
 	static uint32 CurrentFrame;
 };
 
diff --git a/Engine/Source/Runtime/Renderer/Private/PostProcess/PostProcessMorpheus.cpp b/Engine/Source/Runtime/Renderer/Private/PostProcess/PostProcessMorpheus.cpp
--- a/Engine/Source/Runtime/Renderer/Private/PostProcess/PostProcessMorpheus.cpp
+++ b/Engine/Source/Runtime/Renderer/Private/PostProcess/PostProcessMorpheus.cpp
@@ -1,4 +1,4 @@
-// Copyright 1998-2016 Epic Games, Inc. All Rights Reserved.
+// Copyright 1998-2015 Epic Games, Inc. All Rights Reserved.
 
 /*=============================================================================
 	PostProcessMorpheus.cpp: Post processing for Sony Morpheus HMD device.
@@ -54,8 +54,10 @@ public:
 	FShaderParameter TextureScale;
 	FShaderParameter TextureOffset;
 	FShaderParameter TextureUVOffset;
+	FShaderParameter RCoefficients;
+	FShaderParameter GCoefficients;
+	FShaderParameter BCoefficients;
 
-	FShaderResourceParameter DistortionTextureParam; 
 	FShaderResourceParameter DistortionTextureSampler; 
 
 	/** Initialization constructor. */
@@ -73,12 +75,13 @@ public:
 
 		TextureUVOffset.Bind(Initializer.ParameterMap, TEXT("TextureUVOffset"));
 		//check(TextureUVOffset.IsBound());
-			
-		DistortionTextureParam.Bind(Initializer.ParameterMap, TEXT("DistortionTextureArray"));
-		//check(DistortionTextureLeftParam.IsBound());		
 
 		DistortionTextureSampler.Bind(Initializer.ParameterMap, TEXT("DistortionTextureSampler"));
 		//check(DistortionTextureSampler.IsBound());		
+
+		RCoefficients.Bind(Initializer.ParameterMap, TEXT("RCoefficients"));
+		GCoefficients.Bind(Initializer.ParameterMap, TEXT("GCoefficients"));
+		BCoefficients.Bind(Initializer.ParameterMap, TEXT("BCoefficients"));
 	}
 
 
@@ -95,25 +98,27 @@ public:
 			check(GEngine->HMDDevice.IsValid());
 			TSharedPtr< class IHeadMountedDisplay, ESPMode::ThreadSafe > HMDDevice = GEngine->HMDDevice;
 
+			check(HMDDevice->GetHMDDeviceType() == EHMDDeviceType::DT_Morpheus);
+
+			auto RCoefs = HMDDevice->GetRedDistortionParameters();
+			auto GCoefs = HMDDevice->GetGreenDistortionParameters();
+			auto BCoefs = HMDDevice->GetBlueDistortionParameters();
+			for (uint32 i = 0; i < 5; ++i)
+			{
+				SetShaderValue(Context.RHICmdList, ShaderRHI, RCoefficients, RCoefs[i], i);
+				SetShaderValue(Context.RHICmdList, ShaderRHI, GCoefficients, GCoefs[i], i);
+				SetShaderValue(Context.RHICmdList, ShaderRHI, BCoefficients, BCoefs[i], i);
+			}
+
 			check (StereoPass != eSSP_FULL);
 			if (StereoPass == eSSP_LEFT_EYE)
 			{
-				FTexture* TextureLeft = HMDDevice->GetDistortionTextureLeft();
-				if (TextureLeft)
-				{
-					SetTextureParameter(Context.RHICmdList, ShaderRHI, DistortionTextureParam, DistortionTextureSampler, TextureLeft->SamplerStateRHI, TextureLeft->TextureRHI);
-				}
 				SetShaderValue(Context.RHICmdList, ShaderRHI, TextureScale, HMDDevice->GetTextureScaleLeft());
 				SetShaderValue(Context.RHICmdList, ShaderRHI, TextureOffset, HMDDevice->GetTextureOffsetLeft());
 				SetShaderValue(Context.RHICmdList, ShaderRHI, TextureUVOffset, 0.0f);
 			}
 			else
 			{
-				FTexture* TextureRight = HMDDevice->GetDistortionTextureRight();
-				if (TextureRight)
-				{
-					SetTextureParameter(Context.RHICmdList, ShaderRHI, DistortionTextureParam, DistortionTextureSampler, TextureRight->SamplerStateRHI, TextureRight->TextureRHI);
-				}
 				SetShaderValue(Context.RHICmdList, ShaderRHI, TextureScale, HMDDevice->GetTextureScaleRight());
 				SetShaderValue(Context.RHICmdList, ShaderRHI, TextureOffset, HMDDevice->GetTextureOffsetRight());
 				SetShaderValue(Context.RHICmdList, ShaderRHI, TextureUVOffset, -0.5f);
@@ -127,7 +132,7 @@ public:
 	virtual bool Serialize(FArchive& Ar)
 	{
 		bool bShaderHasOutdatedParameters = FGlobalShader::Serialize(Ar);
-		Ar << PostprocessParameter << DeferredParameters << TextureScale << TextureOffset << TextureUVOffset << DistortionTextureParam << DistortionTextureSampler;
+		Ar << PostprocessParameter << DeferredParameters << TextureScale << TextureOffset << TextureUVOffset << RCoefficients << GCoefficients << BCoefficients << DistortionTextureSampler;
 		return bShaderHasOutdatedParameters;
 	}
 };
@@ -251,7 +256,7 @@ void FRCPassPostProcessMorpheus::Process(FRenderingCompositePassContext& Context
 
 FPooledRenderTargetDesc FRCPassPostProcessMorpheus::ComputeOutputDesc(EPassOutputId InPassOutputId) const
 {
-	FPooledRenderTargetDesc Ret = GetInput(ePId_Input0)->GetOutput()->RenderTargetDesc;
+	FPooledRenderTargetDesc Ret = GetInput(static_cast<EPassInputId>(0))->GetOutput()->RenderTargetDesc;
 
 	Ret.NumSamples = 1;	// no MSAA
 	Ret.Reset();
diff --git a/Engine/Source/Runtime/Renderer/Private/PostProcess/SceneRenderTargets.cpp b/Engine/Source/Runtime/Renderer/Private/PostProcess/SceneRenderTargets.cpp
--- a/Engine/Source/Runtime/Renderer/Private/PostProcess/SceneRenderTargets.cpp
+++ b/Engine/Source/Runtime/Renderer/Private/PostProcess/SceneRenderTargets.cpp
@@ -1251,13 +1251,16 @@ void FSceneRenderTargets::FinishRenderingLightAttenuation(FRHICommandList& RHICm
 	GRenderTargetPool.VisualizeTexture.SetCheckPoint(RHICmdList, GetLightAttenuation());
 }
 
-void FSceneRenderTargets::BeginRenderingTranslucency(FRHICommandList& RHICmdList, const FViewInfo& View)
+void FSceneRenderTargets::BeginRenderingTranslucency(FRHICommandList& RHICmdList, const FViewInfo& View, bool bFirstTimeThisFrame)
 {
 	// Use the scene color buffer.
 	BeginRenderingSceneColor(RHICmdList, ESimpleRenderTargetMode::EExistingColorAndDepth, FExclusiveDepthStencil::DepthRead_StencilWrite);
-	
-	// Clear the stencil buffer for ResponsiveAA
-	RHICmdList.Clear(false, FLinearColor::White, false, (float)ERHIZBuffer::FarPlane, true, 0, FIntRect());
+
+	if (bFirstTimeThisFrame)
+	{
+		// Clear the stencil buffer for ResponsiveAA
+		RHICmdList.Clear(false, FLinearColor::White, false, (float)ERHIZBuffer::FarPlane, true, 0, FIntRect());
+	}
 		
 	// viewport to match view size
 	RHICmdList.SetViewport(View.ViewRect.Min.X, View.ViewRect.Min.Y, 0.0f, View.ViewRect.Max.X, View.ViewRect.Max.Y, 1.0f);
diff --git a/Engine/Source/Runtime/Renderer/Private/PostProcess/SceneRenderTargets.h b/Engine/Source/Runtime/Renderer/Private/PostProcess/SceneRenderTargets.h
--- a/Engine/Source/Runtime/Renderer/Private/PostProcess/SceneRenderTargets.h
+++ b/Engine/Source/Runtime/Renderer/Private/PostProcess/SceneRenderTargets.h
@@ -266,7 +266,7 @@ public:
 	/** Resolves the appropriate shadow depth cube map and restores default state. */
 	void FinishRenderingCubeShadowDepth(FRHICommandList& RHICmdList, int32 ShadowResolution);
 	
-	void BeginRenderingTranslucency(FRHICommandList& RHICmdList, const class FViewInfo& View);
+	void BeginRenderingTranslucency(FRHICommandList& RHICmdList, const class FViewInfo& View, bool bFirstTimeThisFrame = true);
 	void FinishRenderingTranslucency(FRHICommandListImmediate& RHICmdList, const class FViewInfo& View);
 
 	bool BeginRenderingSeparateTranslucency(FRHICommandList& RHICmdList, const FViewInfo& View, bool bFirstTimeThisFrame);
diff --git a/Engine/Source/Runtime/Renderer/Private/SceneRendering.cpp b/Engine/Source/Runtime/Renderer/Private/SceneRendering.cpp
--- a/Engine/Source/Runtime/Renderer/Private/SceneRendering.cpp
+++ b/Engine/Source/Runtime/Renderer/Private/SceneRendering.cpp
@@ -1029,7 +1029,7 @@ void FViewInfo::CreateUniformBuffer(
 		GMaxRHIFeatureLevel > ERHIFeatureLevel::ES3_1) ? 1.0f : 0.0f;
 
 	// Padding between the left and right eye may be introduced by an HMD, which instanced stereo needs to account for.
-	if (Family != nullptr && StereoPass == eSSP_LEFT_EYE)
+	if ((Family != nullptr) && (StereoPass == eSSP_LEFT_EYE) && (Family->Views.Num() > 1))
 	{
 		check(Family->Views.Num() == 2);
 		const float FamilySizeX = static_cast<float>(Family->FamilySizeX);
diff --git a/Engine/Source/Runtime/Renderer/Private/TranslucentRendering.cpp b/Engine/Source/Runtime/Renderer/Private/TranslucentRendering.cpp
--- a/Engine/Source/Runtime/Renderer/Private/TranslucentRendering.cpp
+++ b/Engine/Source/Runtime/Renderer/Private/TranslucentRendering.cpp
@@ -204,15 +204,15 @@ static void SetTranslucentRenderTargetAndState(FRHICommandList& RHICmdList, cons
 {
 	FSceneRenderTargets& SceneContext = FSceneRenderTargets::Get(RHICmdList);
 	bool bSetupTranslucentState = true;
+	bool bNeedsClear = (&View == View.Family->Views[0]) && bFirstTimeThisFrame;
 
 	if ((TranslucenyPassType == TPT_SeparateTransluceny) && SceneContext.IsSeparateTranslucencyActive(View))
 	{
-		const bool bNeedsClear = (&View == View.Family->Views[0]) && bFirstTimeThisFrame;
 		bSetupTranslucentState = SceneContext.BeginRenderingSeparateTranslucency(RHICmdList, View, bNeedsClear);
 	}
-	else
+	else if ((TranslucenyPassType == TPT_NonSeparateTransluceny) && !SceneContext.IsSeparateTranslucencyActive(View))
 	{
-		SceneContext.BeginRenderingTranslucency(RHICmdList, View);
+		SceneContext.BeginRenderingTranslucency(RHICmdList, View, bNeedsClear);
 	}
 
 	if (bSetupTranslucentState)
@@ -998,7 +998,7 @@ void FTranslucentPrimSet::AppendScenePrimitives(FSortedPrim* Normal, int32 NumNo
 }
 
 void FTranslucentPrimSet::PlaceScenePrimitive(FPrimitiveSceneInfo* PrimitiveSceneInfo, const FViewInfo& ViewInfo, bool bUseNormalTranslucency, bool bUseSeparateTranslucency, void *NormalPlace, int32& NormalNum, void* SeparatePlace, int32& SeparateNum)
-{	
+{
 	const float SortKey = CalculateTranslucentSortKey(PrimitiveSceneInfo, ViewInfo);
 	const auto FeatureLevel = ViewInfo.GetFeatureLevel();
 	int32 CVarEnabled = FSceneRenderTargets::CVarSetSeperateTranslucencyEnabled.GetValueOnRenderThread();
@@ -1110,7 +1110,7 @@ public:
 	FTranslucencyPassParallelCommandListSet(const FViewInfo& InView, FRHICommandListImmediate& InParentCmdList, bool bInParallelExecute, bool bInCreateSceneContext, ETranslucencyPassType InTranslucenyPassType)
 		: FParallelCommandListSet(GET_STATID(STAT_CLP_Translucency), InView, InParentCmdList, bInParallelExecute, bInCreateSceneContext)
 		, TranslucenyPassType(InTranslucenyPassType)
-		, bFirstTimeThisFrame(InTranslucenyPassType == TPT_SeparateTransluceny)
+		, bFirstTimeThisFrame(true)
 	{
 		SetStateOnCommandList(ParentCmdList);
 	}
diff --git a/Engine/Source/Runtime/Slate/Private/Framework/Application/SlateApplication.cpp b/Engine/Source/Runtime/Slate/Private/Framework/Application/SlateApplication.cpp
--- a/Engine/Source/Runtime/Slate/Private/Framework/Application/SlateApplication.cpp
+++ b/Engine/Source/Runtime/Slate/Private/Framework/Application/SlateApplication.cpp
@@ -2157,29 +2157,14 @@ void FSlateApplication::InvalidateAllViewports()
 	Renderer->InvalidateAllViewports();
 }
 
-
 void FSlateApplication::RegisterGameViewport( TSharedRef<SViewport> InViewport )
 {
 	InViewport->SetActive(true);
 	GameViewportWidget = InViewport;
 	
-	FWidgetPath PathToViewport;
-	// If we cannot find the window it could have been destroyed.
-	if (FSlateWindowHelper::FindPathToWidget(SlateWindows, InViewport, PathToViewport, EVisibility::All))
-	{
-		FReply Reply = FReply::Handled().SetUserFocus(InViewport, EFocusCause::SetDirectly, true);
-	
-		// Set keyboard focus on the actual OS window for the top level Slate window in the viewport path
-		// This is needed because some OS messages are only sent to the window with keyboard focus
-		// Slate will translate the message and send it to the actual widget with focus.
-		// Without this we don't get WM_KEYDOWN or WM_CHAR messages in play in viewport sessions.
-		PathToViewport.GetWindow()->GetNativeWindow()->SetWindowFocus();
-
-		ProcessReply( PathToViewport, Reply, nullptr, nullptr );
-	}
+	ActivateGameViewport();
 }
 
-
 void FSlateApplication::UnregisterGameViewport()
 {
 	ResetToDefaultPointerInputSettings();
@@ -2256,6 +2241,36 @@ void FSlateApplication::SetJoystickCaptorToGameViewport()
 	SetAllUserFocusToGameViewport();
 }
 
+void FSlateApplication::ActivateGameViewport()
+{
+	// Only focus the window if the application is active, if not the application activation sequence will take care of it.
+	if (bAppIsActive && GameViewportWidget.IsValid())
+	{
+		TSharedRef<SViewport> GameViewportWidgetRef = GameViewportWidget.Pin().ToSharedRef();
+		
+		FWidgetPath PathToViewport;
+		// If we cannot find the window it could have been destroyed.
+		if (FSlateWindowHelper::FindPathToWidget(SlateWindows, GameViewportWidgetRef, PathToViewport, EVisibility::All))
+		{
+			TSharedRef<SWindow> Window = PathToViewport.GetWindow();
+
+			// Set keyboard focus on the actual OS window for the top level Slate window in the viewport path
+			// This is needed because some OS messages are only sent to the window with keyboard focus
+			// Slate will translate the message and send it to the actual widget with focus.
+			// Without this we don't get WM_KEYDOWN or WM_CHAR messages in play in viewport sessions.
+			Window->GetNativeWindow()->SetWindowFocus();
+
+			// Activate the viewport and process the reply 
+			FWindowActivateEvent ActivateEvent(FWindowActivateEvent::EA_Activate, Window);
+			FReply ViewportActivatedReply = GameViewportWidgetRef->OnViewportActivated(ActivateEvent);
+			if (ViewportActivatedReply.IsEventHandled())
+			{
+				ProcessReply(PathToViewport, ViewportActivatedReply, nullptr, nullptr);
+			}
+		}
+	}
+}
+
 void FSlateApplication::SetUserFocus(uint32 UserIndex, const TSharedPtr<SWidget>& WidgetToFocus, EFocusCause ReasonFocusIsChanging /* = EFocusCause::SetDirectly*/)
 {
 	const bool bValidWidget = WidgetToFocus.IsValid();
@@ -5550,6 +5565,7 @@ bool FSlateApplication::ProcessWindowActivatedEvent( const FWindowActivateEvent&
 			{
 				if (PathToViewport.GetWindow() == ActiveTopLevelWindow)
 				{
+					// Activate the viewport and process the reply 
 					FReply ViewportActivatedReply = GameViewportWidgetPtr->OnViewportActivated(ActivateEvent);
 					if (ViewportActivatedReply.IsEventHandled())
 					{
diff --git a/Engine/Source/Runtime/Slate/Public/Framework/Application/SlateApplication.h b/Engine/Source/Runtime/Slate/Public/Framework/Application/SlateApplication.h
--- a/Engine/Source/Runtime/Slate/Public/Framework/Application/SlateApplication.h
+++ b/Engine/Source/Runtime/Slate/Public/Framework/Application/SlateApplication.h
@@ -448,6 +448,11 @@ public:
 	DEPRECATED(4.6, "FSlateApplication::SetJoystickCaptorToGameViewport() is deprecated, use FSlateApplication::SetAllUserFocusToGameViewport() instead.")
 	void SetJoystickCaptorToGameViewport();
 
+	/**
+	 * Activates the Game Viewport if it is properly childed under a window
+	 */
+	void ActivateGameViewport();
+
 	/**
 	 * Sets specified user focus to the SWidget passed in.
 	 *
diff --git a/Engine/Source/Runtime/SlateRHIRenderer/Private/Slate3DRenderer.cpp b/Engine/Source/Runtime/SlateRHIRenderer/Private/Slate3DRenderer.cpp
--- a/Engine/Source/Runtime/SlateRHIRenderer/Private/Slate3DRenderer.cpp
+++ b/Engine/Source/Runtime/SlateRHIRenderer/Private/Slate3DRenderer.cpp
@@ -15,17 +15,19 @@ FSlate3DRenderer::FSlate3DRenderer( TSharedRef<FSlateFontServices> InSlateFontSe
 	ElementBatcher = MakeShareable(new FSlateElementBatcher(RenderTargetPolicy.ToSharedRef()));
 }
 
-FSlate3DRenderer::~FSlate3DRenderer()
+void FSlate3DRenderer::Cleanup()
 {
-	if( RenderTargetPolicy.IsValid() )
+	if ( RenderTargetPolicy.IsValid() )
 	{
 		RenderTargetPolicy->ReleaseResources();
 	}
 
-	if ( IsInGameThread() )
-	{
-		FlushRenderingCommands();
-	}
+	BeginCleanup(this);
+}
+
+void FSlate3DRenderer::FinishCleanup()
+{
+	delete this;
 }
 
 FSlateDrawBuffer& FSlate3DRenderer::GetDrawBuffer()
diff --git a/Engine/Source/Runtime/SlateRHIRenderer/Private/Slate3DRenderer.h b/Engine/Source/Runtime/SlateRHIRenderer/Private/Slate3DRenderer.h
--- a/Engine/Source/Runtime/SlateRHIRenderer/Private/Slate3DRenderer.h
+++ b/Engine/Source/Runtime/SlateRHIRenderer/Private/Slate3DRenderer.h
@@ -4,11 +4,13 @@
 
 #include "ISlate3DRenderer.h"
 
-class FSlate3DRenderer : public ISlate3DRenderer
+class FSlate3DRenderer : public ISlate3DRenderer, public FDeferredCleanupInterface
 {
 public:
 	FSlate3DRenderer( TSharedRef<FSlateFontServices> InSlateFontServices, TSharedRef<FSlateRHIResourceManager> InResourceManager, bool bUseGammaCorrection = false );
-	~FSlate3DRenderer();
+
+	void Cleanup();
+	virtual void FinishCleanup() override;
 
 	virtual FSlateDrawBuffer& GetDrawBuffer() override;
 	virtual void DrawWindow_GameThread(FSlateDrawBuffer& DrawBuffer) override;
diff --git a/Engine/Source/Runtime/SlateRHIRenderer/Private/SlateRHIRenderer.cpp b/Engine/Source/Runtime/SlateRHIRenderer/Private/SlateRHIRenderer.cpp
--- a/Engine/Source/Runtime/SlateRHIRenderer/Private/SlateRHIRenderer.cpp
+++ b/Engine/Source/Runtime/SlateRHIRenderer/Private/SlateRHIRenderer.cpp
@@ -1309,7 +1309,7 @@ void FSlateRHIRenderer::ReleaseCachingResourcesFor(const ILayoutCache* Cacher)
 		ReleaseCachingResourcesFor,
 		FReleaseCachingResourcesForContext, Context, MarshalContext,
 		{
-			Context.RenderPolicy->ReleaseCachingResourcesFor(Context.Cacher);
+			Context.RenderPolicy->ReleaseCachingResourcesFor(RHICmdList, Context.Cacher);
 		});
 }
 
diff --git a/Engine/Source/Runtime/SlateRHIRenderer/Private/SlateRHIRendererModule.cpp b/Engine/Source/Runtime/SlateRHIRenderer/Private/SlateRHIRendererModule.cpp
--- a/Engine/Source/Runtime/SlateRHIRenderer/Private/SlateRHIRendererModule.cpp
+++ b/Engine/Source/Runtime/SlateRHIRenderer/Private/SlateRHIRendererModule.cpp
@@ -64,7 +64,9 @@ public:
 	{
 		ConditionalCreateResources();
 
-		return MakeShareable( new FSlate3DRenderer( SlateFontServices.ToSharedRef(),  ResourceManager.ToSharedRef(), bUseGammaCorrection ) );
+		return MakeShareable(new FSlate3DRenderer(SlateFontServices.ToSharedRef(), ResourceManager.ToSharedRef(), bUseGammaCorrection), [=] (FSlate3DRenderer* Renderer) {
+			Renderer->Cleanup();
+		});
 	}
 
 	virtual TSharedRef<ISlateFontAtlasFactory> CreateSlateFontAtlasFactory() override
diff --git a/Engine/Source/Runtime/SlateRHIRenderer/Private/SlateRHIRenderingPolicy.cpp b/Engine/Source/Runtime/SlateRHIRenderer/Private/SlateRHIRenderingPolicy.cpp
--- a/Engine/Source/Runtime/SlateRHIRenderer/Private/SlateRHIRenderingPolicy.cpp
+++ b/Engine/Source/Runtime/SlateRHIRenderer/Private/SlateRHIRenderingPolicy.cpp
@@ -33,17 +33,6 @@ FSlateRHIRenderingPolicy::FSlateRHIRenderingPolicy(TSharedRef<FSlateFontServices
 	SetDefaultBlendMode(FBlendStateInitializerRHI::FRenderTarget(BO_Add, BF_SourceAlpha, BF_InverseSourceAlpha, BO_Add, BF_One, BF_InverseSourceAlpha, CW_RGBA));
 }
 
-FSlateRHIRenderingPolicy::~FSlateRHIRenderingPolicy()
-{
-	// Delete released resources.  Note this MUST NOT be called before the rendering resources have been released
-	ReleaseResources();
-
-	if ( IsInGameThread() )
-	{
-		FlushRenderingCommands();
-	}
-}
-
 void FSlateRHIRenderingPolicy::InitResources()
 {
 	int32 NumVertices = 100;
@@ -134,9 +123,9 @@ void FSlateRHIRenderingPolicy::UpdateVertexAndIndexBuffers(FRHICommandListImmedi
 	UpdateVertexAndIndexBuffers(RHICmdList, InBatchData, Buffers->VertexBuffer, Buffers->IndexBuffer);
 }
 
-void FSlateRHIRenderingPolicy::ReleaseCachingResourcesFor(const ILayoutCache* Cacher)
+void FSlateRHIRenderingPolicy::ReleaseCachingResourcesFor(FRHICommandListImmediate& RHICmdList, const ILayoutCache* Cacher)
 {
-	ResourceManager->ReleaseCachingResourcesFor(Cacher);
+	ResourceManager->ReleaseCachingResourcesFor(RHICmdList, Cacher);
 }
 
 void FSlateRHIRenderingPolicy::UpdateVertexAndIndexBuffers(FRHICommandListImmediate& RHICmdList, FSlateBatchData& InBatchData, TSlateElementVertexBuffer<FSlateVertex>& VertexBuffer, FSlateElementIndexBuffer& IndexBuffer)
diff --git a/Engine/Source/Runtime/SlateRHIRenderer/Private/SlateRHIRenderingPolicy.h b/Engine/Source/Runtime/SlateRHIRenderer/Private/SlateRHIRenderingPolicy.h
--- a/Engine/Source/Runtime/SlateRHIRenderer/Private/SlateRHIRenderingPolicy.h
+++ b/Engine/Source/Runtime/SlateRHIRenderer/Private/SlateRHIRenderingPolicy.h
@@ -11,12 +11,11 @@ class FSlateRHIRenderingPolicy : public FSlateRenderingPolicy
 {
 public:
 	FSlateRHIRenderingPolicy(TSharedRef<FSlateFontServices> InSlateFontServices, TSharedRef<FSlateRHIResourceManager> InResourceManager, TOptional<int32> InitialBufferSize = TOptional<int32>());
-	virtual ~FSlateRHIRenderingPolicy();
 
 	void UpdateVertexAndIndexBuffers(FRHICommandListImmediate& RHICmdList, FSlateBatchData& BatchData);
 	void UpdateVertexAndIndexBuffers(FRHICommandListImmediate& RHICmdList, FSlateBatchData& BatchData, const TSharedRef<FSlateRenderDataHandle, ESPMode::ThreadSafe>& RenderHandle);
 
-	void ReleaseCachingResourcesFor(const ILayoutCache* Cacher);
+	void ReleaseCachingResourcesFor(FRHICommandListImmediate& RHICmdList, const ILayoutCache* Cacher);
 
 	void SetDefaultBlendMode(const FBlendStateInitializerRHI& BlendState);
 
diff --git a/Engine/Source/Runtime/SlateRHIRenderer/Private/SlateRHIResourceManager.cpp b/Engine/Source/Runtime/SlateRHIRenderer/Private/SlateRHIResourceManager.cpp
--- a/Engine/Source/Runtime/SlateRHIRenderer/Private/SlateRHIResourceManager.cpp
+++ b/Engine/Source/Runtime/SlateRHIRenderer/Private/SlateRHIResourceManager.cpp
@@ -224,6 +224,42 @@ bool FSlateRHIResourceManager::IsAtlasPageResourceAlphaOnly() const
 	return false;
 }
 
+void FSlateRHIResourceManager::Tick(float DeltaSeconds)
+{
+	// Don't need to do this if there's no RHI thread.
+	if ( GRHIThread )
+	{
+		struct FDeleteCachedRenderDataContext
+		{
+			FSlateRHIResourceManager* ResourceManager;
+		};
+		FDeleteCachedRenderDataContext DeleteCachedRenderDataContext =
+		{
+			this,
+		};
+		ENQUEUE_UNIQUE_RENDER_COMMAND_ONEPARAMETER(
+			DeleteCachedRenderData,
+			FDeleteCachedRenderDataContext, Context, DeleteCachedRenderDataContext,
+			{
+				// Go through the pending delete buffers and see if any of their fences has cleared
+				// the RHI thread, if so, they should be safe to delete now.
+				for ( int32 BufferIndex = Context.ResourceManager->PooledBuffersPendingRelease.Num() - 1; BufferIndex >= 0; BufferIndex-- )
+				{
+					FCachedRenderBuffers* PooledBuffer = Context.ResourceManager->PooledBuffersPendingRelease[BufferIndex];
+
+					if ( PooledBuffer->ReleaseResourcesFence->IsComplete() )
+					{
+						PooledBuffer->VertexBuffer.Destroy();
+						PooledBuffer->IndexBuffer.Destroy();
+						delete PooledBuffer;
+
+						Context.ResourceManager->PooledBuffersPendingRelease.RemoveAt(BufferIndex);
+					}
+				}
+			});
+	}
+}
+
 void FSlateRHIResourceManager::AddReferencedObjects(FReferenceCollector& Collector)
 {
 	for ( TSet<UObject*>* AccessedUObjects : AllAccessedUObject )
@@ -852,6 +888,7 @@ void FSlateRHIResourceManager::UpdateTextureAtlases()
 
 FCachedRenderBuffers* FSlateRHIResourceManager::FindOrCreateCachedBuffersForHandle(const TSharedRef<FSlateRenderDataHandle, ESPMode::ThreadSafe>& RenderHandle)
 {
+	// Should only be called by the rendering thread
 	check(IsInRenderingThread());
 
 	FCachedRenderBuffers* Buffers = CachedBuffers.FindRef(&RenderHandle.Get());
@@ -902,13 +939,12 @@ void FSlateRHIResourceManager::BeginReleasingRenderData(const FSlateRenderDataHa
 		ReleaseCachedRenderData,
 		FReleaseCachedRenderDataContext, Context, ReleaseCachedRenderDataContext,
 		{
-			Context.ResourceManager->ReleaseCachedRenderData(Context.RenderDataHandle, Context.LayoutCacher);
+			Context.ResourceManager->ReleaseCachedRenderData(RHICmdList, Context.RenderDataHandle, Context.LayoutCacher);
 		});
 }
 
-void FSlateRHIResourceManager::ReleaseCachedRenderData(const FSlateRenderDataHandle* RenderHandle, const ILayoutCache* LayoutCacher)
+void FSlateRHIResourceManager::ReleaseCachedRenderData(FRHICommandListImmediate& RHICmdList, const FSlateRenderDataHandle* RenderHandle, const ILayoutCache* LayoutCacher)
 {
-	// Should only be called by the rendering thread
 	check(IsInRenderingThread());
 
 	FCachedRenderBuffers* PooledBuffer = CachedBuffers.FindRef(RenderHandle);
@@ -921,17 +957,14 @@ void FSlateRHIResourceManager::ReleaseCachedRenderData(const FSlateRenderDataHan
 		}
 		else
 		{
-			// The buffer pool may have already been released, so lets just delete this buffer.
-			PooledBuffer->VertexBuffer.Destroy();
-			PooledBuffer->IndexBuffer.Destroy();
-			delete PooledBuffer;
+			ReleaseCachedBuffer(RHICmdList, PooledBuffer);
 		}
 
 		CachedBuffers.Remove(RenderHandle);
 	}
 }
 
-void FSlateRHIResourceManager::ReleaseCachingResourcesFor(const ILayoutCache* Cacher)
+void FSlateRHIResourceManager::ReleaseCachingResourcesFor(FRHICommandListImmediate& RHICmdList, const ILayoutCache* Cacher)
 {
 	check(IsInRenderingThread());
 
@@ -940,15 +973,30 @@ void FSlateRHIResourceManager::ReleaseCachingResourcesFor(const ILayoutCache* Ca
 	{
 		for ( FCachedRenderBuffers* PooledBuffer : *Pool )
 		{
-			PooledBuffer->VertexBuffer.Destroy();
-			PooledBuffer->IndexBuffer.Destroy();
-			delete PooledBuffer;
+			ReleaseCachedBuffer(RHICmdList, PooledBuffer);
 		}
 
 		CachedBufferPool.Remove(Cacher);
 	}
 }
 
+void FSlateRHIResourceManager::ReleaseCachedBuffer(FRHICommandListImmediate& RHICmdList, FCachedRenderBuffers* PooledBuffer)
+{
+	check(IsInRenderingThread());
+
+	if ( GRHIThread )
+	{
+		PooledBuffersPendingRelease.Add(PooledBuffer);
+		PooledBuffer->ReleaseResourcesFence = RHICmdList.RHIThreadFence();
+	}
+	else
+	{
+		PooledBuffer->VertexBuffer.Destroy();
+		PooledBuffer->IndexBuffer.Destroy();
+		delete PooledBuffer;
+	}
+}
+
 void FSlateRHIResourceManager::ReleaseResources()
 {
 	checkSlow( IsThreadSafeForSlateRendering() );
diff --git a/Engine/Source/Runtime/SlateRHIRenderer/Private/SlateRHIResourceManager.h b/Engine/Source/Runtime/SlateRHIRenderer/Private/SlateRHIResourceManager.h
--- a/Engine/Source/Runtime/SlateRHIRenderer/Private/SlateRHIResourceManager.h
+++ b/Engine/Source/Runtime/SlateRHIRenderer/Private/SlateRHIResourceManager.h
@@ -67,28 +67,37 @@ struct FCachedRenderBuffers
 {
 	TSlateElementVertexBuffer<FSlateVertex> VertexBuffer;
 	FSlateElementIndexBuffer IndexBuffer;
+
+	FGraphEventRef ReleaseResourcesFence;
 };
 
 
 /**
  * Stores a mapping of texture names to their RHI texture resource               
  */
-class FSlateRHIResourceManager : public ISlateAtlasProvider, public ISlateRenderDataManager, public FSlateShaderResourceManager, public FGCObject
+class FSlateRHIResourceManager : public ISlateAtlasProvider, public ISlateRenderDataManager, public FSlateShaderResourceManager, public FTickableGameObject, public FGCObject
 {
 public:
 	FSlateRHIResourceManager();
 	virtual ~FSlateRHIResourceManager();
 
-	/** ISlateAtlasProvider */
+	/** ISlateAtlasProvider interface */
 	virtual int32 GetNumAtlasPages() const override;
 	virtual FIntPoint GetAtlasPageSize() const override;
 	virtual FSlateShaderResource* GetAtlasPageResource(const int32 InIndex) const override;
 	virtual bool IsAtlasPageResourceAlphaOnly() const override;
 
-	/** ISlateRenderDataManager */
+	/** ISlateRenderDataManager interface */
 	virtual void BeginReleasingRenderData(const FSlateRenderDataHandle* RenderHandle) override;
 
-	/** FGCObject */
+	/** FTickableGameObject interface */
+	virtual bool IsTickable() const override { return true; }
+	virtual bool IsTickableWhenPaused() const override { return true; }
+	virtual bool IsTickableInEditor() const override { return true; }
+	virtual TStatId GetStatId() const override { RETURN_QUICK_DECLARE_CYCLE_STAT(FSlateRHIResourceManager, STATGROUP_Tickables); }
+	virtual void Tick(float DeltaSeconds) override;
+
+	/** FGCObject interface */
 	virtual void AddReferencedObjects(FReferenceCollector& Collector) override;
 
 	/**
@@ -172,7 +181,7 @@ public:
 	/**
 	 * Releases all cached buffers allocated by a given layout cacher.  This would happen when an invalidation panel is destroyed.
 	 */
-	void ReleaseCachingResourcesFor(const ILayoutCache* Cacher);
+	void ReleaseCachingResourcesFor(FRHICommandListImmediate& RHICmdList, const ILayoutCache* Cacher);
 
 	/**
 	 * Releases rendering resources
@@ -187,11 +196,13 @@ public:
 	void ReloadTextures();
 
 private:
+	void ReleaseCachedBuffer(FRHICommandListImmediate& RHICmdList, FCachedRenderBuffers* PooledBuffer);
+
 	/**
 	* Releases the cached render data for a given render handle.  The layout cacher that owned the handle must also be
 	* provided, as render handle is likely no longer a valid object.
 	*/
-	void ReleaseCachedRenderData(const FSlateRenderDataHandle* RenderHandle, const ILayoutCache* LayoutCacher);
+	void ReleaseCachedRenderData(FRHICommandListImmediate& RHICmdList, const FSlateRenderDataHandle* RenderHandle, const ILayoutCache* LayoutCacher);
 
 private:
 	/**
@@ -288,5 +299,11 @@ private:
 
 	typedef TMap< const ILayoutCache*, TArray< FCachedRenderBuffers* > > TCachedBufferPoolMap;
 	TCachedBufferPoolMap CachedBufferPool;
+
+	/**
+	 * Holds onto a list of pooled buffers that are no longer being used but still need 
+	 * to be deleted after the RHI thread is done with them.
+	 */
+	TArray<FCachedRenderBuffers*> PooledBuffersPendingRelease;
 };
 
diff --git a/Engine/Source/Runtime/Sockets/Private/BSDSockets/SocketSubsystemBSD.cpp b/Engine/Source/Runtime/Sockets/Private/BSDSockets/SocketSubsystemBSD.cpp
--- a/Engine/Source/Runtime/Sockets/Private/BSDSockets/SocketSubsystemBSD.cpp
+++ b/Engine/Source/Runtime/Sockets/Private/BSDSockets/SocketSubsystemBSD.cpp
@@ -191,6 +191,7 @@ ESocketErrors FSocketSubsystemBSD::TranslateErrorCode(int32 Code)
 	case EDQUOT: return SE_EDQUOT;
 	case ESTALE: return SE_ESTALE;
 	case EREMOTE: return SE_EREMOTE;
+	case ENODEV: return SE_NODEV;
 #if !PLATFORM_HAS_NO_EPROCLIM
 	case EPROCLIM: return SE_EPROCLIM;
 #endif
diff --git a/Engine/Source/Runtime/Sockets/Public/SocketTypes.h b/Engine/Source/Runtime/Sockets/Public/SocketTypes.h
--- a/Engine/Source/Runtime/Sockets/Public/SocketTypes.h
+++ b/Engine/Source/Runtime/Sockets/Public/SocketTypes.h
@@ -80,6 +80,7 @@ enum ESocketErrors
 	SE_UDP_ERR_PORT_UNREACH,
 	SE_ADDRFAMILY,
 	SE_SYSTEM,
+	SE_NODEV,
 
 	// this is a special error which means to lookup the most recent error (via GetLastErrorCode())
 	SE_GET_LAST_ERROR_CODE,
