Copying //UE4/Orion-Staging to //UE4/Main (Source: //Orion/Dev-General @ 3196521)

diff --git a/Engine/Source/Runtime/Online/XMPP/Private/XmppJingle/XmppPresenceJingle.cpp b/Engine/Source/Runtime/Online/XMPP/Private/XmppJingle/XmppPresenceJingle.cpp
--- a/Engine/Source/Runtime/Online/XMPP/Private/XmppJingle/XmppPresenceJingle.cpp
+++ b/Engine/Source/Runtime/Online/XMPP/Private/XmppJingle/XmppPresenceJingle.cpp
@@ -219,13 +219,13 @@ protected:
 	}
 };
 
-class FXmppPresenceOutTask : public buzz::XmppTask 
+class FXmppPresenceOutTask : public buzz::XmppTask
 {
 public:
 	explicit FXmppPresenceOutTask(buzz::XmppTaskParentInterface* Parent)
-		: buzz::XmppTask(Parent) 
+		: buzz::XmppTask(Parent)
 	{}
-	virtual ~FXmppPresenceOutTask() 
+	virtual ~FXmppPresenceOutTask()
 	{}
 
 	buzz::XmppReturnStatus Send(const buzz::PresenceStatus& Status)
@@ -233,9 +233,8 @@ public:
 		if (GetState() != STATE_INIT && GetState() != STATE_START)
 			return buzz::XMPP_RETURN_BADSTATE;
 
-		buzz::XmlElement* Presence = TranslateStatus(Status);
-		QueueStanza(Presence);
-		delete Presence;
+		TUniquePtr<buzz::XmlElement> Presence{TranslateStatus(Status)};
+		QueueStanza(Presence.Get());
 		return buzz::XMPP_RETURN_OK;
 	}
 
@@ -244,10 +243,9 @@ public:
 		if (GetState() != STATE_INIT && GetState() != STATE_START)
 			return buzz::XMPP_RETURN_BADSTATE;
 
-		buzz::XmlElement* Presence = TranslateStatus(Status);
+		TUniquePtr<buzz::XmlElement> Presence{TranslateStatus(Status)};
 		Presence->AddAttr(buzz::QN_TO, Jid.Str());
-		QueueStanza(Presence);
-		delete Presence;
+		QueueStanza(Presence.Get());
 		return buzz::XMPP_RETURN_OK;
 	}
 
@@ -256,12 +254,14 @@ public:
 		if (GetState() != STATE_INIT && GetState() != STATE_START)
 			return buzz::XMPP_RETURN_BADSTATE;
 
-		buzz::XmlElement* Presence = new buzz::XmlElement(buzz::QN_PRESENCE);
-		Presence->AddAttr(buzz::QN_TO, Jid.Str());
-		Presence->AddAttr(buzz::QN_TYPE, "probe");
+		buzz::XmlElement Presence{buzz::QN_PRESENCE};
+		Presence.AddAttr(buzz::QN_TO, Jid.Str());
+		Presence.AddAttr(buzz::QN_TYPE, "probe");
+
+		// Add CorrelationID for tracking purposes
+		FXmppJingle::AddCorrIdToStanza(Presence);
 
-		QueueStanza(Presence);
-		delete Presence;
+		QueueStanza(&Presence);
 		return buzz::XMPP_RETURN_OK;
 	}
 
@@ -281,17 +281,21 @@ private:
 	buzz::XmlElement* TranslateStatus(const buzz::PresenceStatus& Status)
 	{
 		buzz::XmlElement* Result = new buzz::XmlElement(buzz::QN_PRESENCE);
-		if (!Status.available()) 
+
+		// Add CorrelationID for tracking purposes
+		FXmppJingle::AddCorrIdToStanza(*Result);
+
+		if (!Status.available())
 		{
 			Result->AddAttr(buzz::QN_TYPE, buzz::STR_UNAVAILABLE);
 		}
-		else 
+		else
 		{
 			if (Status.show() != buzz::PresenceStatus::SHOW_ONLINE &&
-				Status.show() != buzz::PresenceStatus::SHOW_OFFLINE) 
+				Status.show() != buzz::PresenceStatus::SHOW_OFFLINE)
 			{
 				Result->AddElement(new buzz::XmlElement(buzz::QN_SHOW));
-				switch (Status.show()) 
+				switch (Status.show())
 				{
 				default:
 					Result->AddText(buzz::STR_SHOW_AWAY, 1);
@@ -311,7 +315,7 @@ private:
 			Result->AddElement(new buzz::XmlElement(buzz::QN_STATUS));
 			Result->AddText(Status.status(), 1);
 
-			if (!Status.nick().empty()) 
+			if (!Status.nick().empty())
 			{
 				Result->AddElement(new buzz::XmlElement(buzz::QN_NICKNAME));
 				Result->AddText(Status.nick(), 1);
@@ -323,7 +327,7 @@ private:
 			Result->AddElement(new buzz::XmlElement(buzz::QN_PRIORITY));
 			Result->AddText(Pri, 1);
 
-			if (Status.know_capabilities()) 
+			if (Status.know_capabilities())
 			{
 				Result->AddElement(new buzz::XmlElement(buzz::QN_CAPS_C, true));
 				Result->AddAttr(buzz::QN_NODE, Status.caps_node(), 1);
diff --git a/Engine/Source/Runtime/PacketHandlers/PacketHandler/Private/PacketHandler.cpp b/Engine/Source/Runtime/PacketHandlers/PacketHandler/Private/PacketHandler.cpp
--- a/Engine/Source/Runtime/PacketHandlers/PacketHandler/Private/PacketHandler.cpp
+++ b/Engine/Source/Runtime/PacketHandlers/PacketHandler/Private/PacketHandler.cpp
@@ -324,12 +324,12 @@ const ProcessedPacket PacketHandler::Incoming_Internal(uint8* Packet, int32 Coun
 		}
 		else
 		{
-			return ProcessedPacket();
+			return ProcessedPacket(nullptr, 0, true);
 		}
 	}
 	else
 	{
-		return ProcessedPacket();
+		return ProcessedPacket(nullptr, 0, true);
 	}
 }
 
@@ -417,7 +417,7 @@ const ProcessedPacket PacketHandler::Outgoing_Internal(uint8* Packet, int32 Coun
 		}
 		else
 		{
-			return ProcessedPacket(Packet, CountBits);
+			return ProcessedPacket(nullptr, 0, true);
 		}
 	}
 	else
diff --git a/Engine/Source/Runtime/PacketHandlers/PacketHandler/Public/PacketHandler.h b/Engine/Source/Runtime/PacketHandlers/PacketHandler/Public/PacketHandler.h
--- a/Engine/Source/Runtime/PacketHandlers/PacketHandler/Public/PacketHandler.h
+++ b/Engine/Source/Runtime/PacketHandlers/PacketHandler/Public/PacketHandler.h
@@ -61,19 +61,24 @@ struct PACKETHANDLER_API ProcessedPacket
 	/** Size of the returned packet data in bits */
 	int32 CountBits;
 
+	/** Whether or not there was an error processing the packet */
+	bool bError;
+
 public:
 	ProcessedPacket()
 		: Data()
 		, CountBits(0)
+		, bError(false)
 	{
 	}
 
 	/**
 	 * Base constructor
 	 */
-	ProcessedPacket(uint8* InData, int32 InCountBits)
+	ProcessedPacket(uint8* InData, int32 InCountBits, bool bInError=false)
 		: Data(InData)
 		, CountBits(InCountBits)
+		, bError(bInError)
 	{
 	}
 };
diff --git a/Engine/Source/Runtime/RHI/Private/GPUDefragAllocator.cpp b/Engine/Source/Runtime/RHI/Private/GPUDefragAllocator.cpp
--- a/Engine/Source/Runtime/RHI/Private/GPUDefragAllocator.cpp
+++ b/Engine/Source/Runtime/RHI/Private/GPUDefragAllocator.cpp
@@ -1781,6 +1781,7 @@ FGPUDefragAllocator::EMemoryElementType FGPUDefragAllocator::GetChunkType(FMemor
 */
 int32 FGPUDefragAllocator::GetLargestAvailableAllocation(int32* OutNumFreeChunks/*=nullptr*/)
 {
+	FScopeLock SyncLock(&SynchronizationObject);
 	int32 NumFreeChunks = 0;
 	int64 LargestChunkSize = 0;
 	FMemoryChunk* FreeChunk = FirstFreeChunk;
diff --git a/Engine/Source/Runtime/RenderCore/Private/RenderingThread.cpp b/Engine/Source/Runtime/RenderCore/Private/RenderingThread.cpp
--- a/Engine/Source/Runtime/RenderCore/Private/RenderingThread.cpp
+++ b/Engine/Source/Runtime/RenderCore/Private/RenderingThread.cpp
@@ -12,7 +12,7 @@
 #include "ExceptionHandling.h"
 #include "TaskGraphInterfaces.h"
 #include "StatsData.h"
-
+#include "HAL/ThreadHeartBeat.h"
 //
 // Globals
 //
@@ -857,6 +857,14 @@ static FAutoConsoleVariableRef CVarTimeToBlockOnRenderFence(
 	TEXT("Number of milliseconds the game thread should block when waiting on a render thread fence.")
 	);
 
+
+static int32 GTimeoutForBlockOnRenderFence = 30000;
+static FAutoConsoleVariableRef CVarTimeoutForBlockOnRenderFence(
+	TEXT("g.TimeoutForBlockOnRenderFence"),
+	GTimeoutForBlockOnRenderFence,
+	TEXT("Number of milliseconds the game thread should wait before failing when waiting on a render thread fence.")
+);
+
 /**
  * Block the game thread waiting for a task to finish on the rendering thread.
  */
@@ -895,6 +903,10 @@ static void GameThreadWaitForTask(const FGraphEventRef& Task, bool bEmptyGameThr
 			// rendering thread hasn't crashed :)
 			bool bDone;
 			uint32 WaitTime = FMath::Clamp<uint32>(GTimeToBlockOnRenderFence, 0, 33);
+
+			const double StartTime = FPlatformTime::Seconds();
+			const double EndTime = StartTime + (GTimeoutForBlockOnRenderFence / 1000.0);
+
 			do
 			{
 				CheckRenderingThreadHealth();
@@ -904,6 +916,16 @@ static void GameThreadWaitForTask(const FGraphEventRef& Task, bool bEmptyGameThr
 					FTaskGraphInterface::Get().ProcessThreadUntilIdle(ENamedThreads::GameThread);
 				}
 				bDone = Event->Wait(WaitTime);
+
+				if (!bDone && !GIsEditor) // editor threads can block for quite a while...
+				{
+					// Fatal timeout if we run out of time and this thread is being monitor for heartbeats
+					// (We could just let the heartbeat monitor error for us, but this leads to better diagnostics).
+					if (FPlatformTime::Seconds() >= EndTime && FThreadHeartBeat::Get().IsBeating())
+					{
+						UE_LOG(LogRendererCore, Fatal, TEXT("GameThread timed out waiting for RenderThread after %.02f secs"), FPlatformTime::Seconds() - StartTime);
+					}
+				}
 			}
 			while (!bDone);
 
diff --git a/Engine/Source/Runtime/RenderCore/Public/RenderUtils.h b/Engine/Source/Runtime/RenderCore/Public/RenderUtils.h
--- a/Engine/Source/Runtime/RenderCore/Public/RenderUtils.h
+++ b/Engine/Source/Runtime/RenderCore/Public/RenderUtils.h
@@ -320,6 +320,11 @@ inline bool IsAnyForwardShadingEnabled(EShaderPlatform Platform)
 	return IsForwardShadingEnabled(GetMaxSupportedFeatureLevel(Platform)) || IsSimpleForwardShadingEnabled(Platform);
 }
 
+inline bool IsUsingGBuffers(EShaderPlatform Platform)
+{
+	return !IsAnyForwardShadingEnabled(Platform);
+}
+
 /** Unit cube vertex buffer (VertexDeclarationFVector4) */
 RENDERCORE_API FVertexBufferRHIRef& GetUnitCubeVertexBuffer();
 
diff --git a/Engine/Source/Runtime/Renderer/Private/DeferredShadingRenderer.cpp b/Engine/Source/Runtime/Renderer/Private/DeferredShadingRenderer.cpp
--- a/Engine/Source/Runtime/Renderer/Private/DeferredShadingRenderer.cpp
+++ b/Engine/Source/Runtime/Renderer/Private/DeferredShadingRenderer.cpp
@@ -628,7 +628,7 @@ void FDeferredShadingSceneRenderer::Render(FRHICommandListImmediate& RHICmdList)
 	bool bRequiresRHIClear = true;
 	bool bRequiresFarZQuadClear = false;
 
-	const bool bUseGBuffer = !IsAnyForwardShadingEnabled(GetFeatureLevelShaderPlatform(FeatureLevel));
+	const bool bUseGBuffer = IsUsingGBuffers(GetFeatureLevelShaderPlatform(FeatureLevel));
 
 	if (ClearMethodCVar)
 	{
@@ -765,12 +765,10 @@ void FDeferredShadingSceneRenderer::Render(FRHICommandListImmediate& RHICmdList)
 	const bool bUseVelocityGBuffer = FVelocityRendering::OutputsToGBuffer();
 	const bool bUseSelectiveBasePassOutputs = UseSelectiveBasePassOutputs();
 
-	if (bUseGBuffer)
-	{
-		SCOPE_CYCLE_COUNTER(STAT_FDeferredShadingSceneRenderer_AllocGBufferTargets);
-		SceneContext.PreallocGBufferTargets(bUseVelocityGBuffer); // Even if !bShouldRenderVelocities, the velocity buffer must be bound because it's a compile time option for the shader.
-		SceneContext.AllocGBufferTargets(RHICmdList);
-	}
+	
+	SCOPE_CYCLE_COUNTER(STAT_FDeferredShadingSceneRenderer_AllocGBufferTargets);
+	SceneContext.PreallocGBufferTargets(bUseVelocityGBuffer); // Even if !bShouldRenderVelocities, the velocity buffer must be bound because it's a compile time option for the shader.
+	SceneContext.AllocGBufferTargets(RHICmdList);
 
 	//occlusion can't run before basepass if there's no prepass to fill in some depth to occlude against.
 	bool bOcclusionBeforeBasePass = ((CVarOcclusionQueryLocation.GetValueOnRenderThread() == 1) && bNeedsPrePass) || IsForwardShadingEnabled(FeatureLevel);	
diff --git a/Engine/Source/Runtime/Renderer/Private/MobileShadingRenderer.cpp b/Engine/Source/Runtime/Renderer/Private/MobileShadingRenderer.cpp
--- a/Engine/Source/Runtime/Renderer/Private/MobileShadingRenderer.cpp
+++ b/Engine/Source/Runtime/Renderer/Private/MobileShadingRenderer.cpp
@@ -86,7 +86,7 @@ void FMobileSceneRenderer::InitViews(FRHICommandListImmediate& RHICmdList)
 
 	UpdatePostProcessUsageFlags();
 	
-	OnStartFrame();
+	OnStartFrame(RHICmdList);
 }
 
 /** 
diff --git a/Engine/Source/Runtime/Renderer/Private/PostProcess/PostProcessBufferInspector.cpp b/Engine/Source/Runtime/Renderer/Private/PostProcess/PostProcessBufferInspector.cpp
--- a/Engine/Source/Runtime/Renderer/Private/PostProcess/PostProcessBufferInspector.cpp
+++ b/Engine/Source/Runtime/Renderer/Private/PostProcess/PostProcessBufferInspector.cpp
@@ -15,7 +15,7 @@
 FRCPassPostProcessBufferInspector::FRCPassPostProcessBufferInspector(FRHICommandList& RHICmdList)
 {
 	// AdjustGBufferRefCount(-1) call is done when the pass gets executed
-	FSceneRenderTargets::Get_Todo_PassContext().AdjustGBufferRefCount(RHICmdList, 1);
+	FSceneRenderTargets::Get(RHICmdList).AdjustGBufferRefCount(RHICmdList, 1);
 }
 
 FShader* FRCPassPostProcessBufferInspector::SetShaderTempl(const FRenderingCompositePassContext& Context)
diff --git a/Engine/Source/Runtime/Renderer/Private/PostProcess/PostProcessGBufferHints.cpp b/Engine/Source/Runtime/Renderer/Private/PostProcess/PostProcessGBufferHints.cpp
--- a/Engine/Source/Runtime/Renderer/Private/PostProcess/PostProcessGBufferHints.cpp
+++ b/Engine/Source/Runtime/Renderer/Private/PostProcess/PostProcessGBufferHints.cpp
@@ -81,7 +81,7 @@ IMPLEMENT_SHADER_TYPE(,FPostProcessGBufferHintsPS,TEXT("PostProcessGBufferHints"
 FRCPassPostProcessGBufferHints::FRCPassPostProcessGBufferHints(FRHICommandList& RHICmdList)
 {
 	// AdjustGBufferRefCount(-1) call is done when the pass gets executed
-	FSceneRenderTargets::Get_Todo_PassContext().AdjustGBufferRefCount(RHICmdList, 1);
+	FSceneRenderTargets::Get(RHICmdList).AdjustGBufferRefCount(RHICmdList, 1);
 }
 
 
diff --git a/Engine/Source/Runtime/Renderer/Private/PostProcess/PostProcessing.cpp b/Engine/Source/Runtime/Renderer/Private/PostProcess/PostProcessing.cpp
--- a/Engine/Source/Runtime/Renderer/Private/PostProcess/PostProcessing.cpp
+++ b/Engine/Source/Runtime/Renderer/Private/PostProcess/PostProcessing.cpp
@@ -170,7 +170,7 @@ FPostprocessContext::FPostprocessContext(FRHICommandListImmediate& InRHICmdList,
 	, SceneColor(0)
 	, SceneDepth(0)
 {
-	FSceneRenderTargets& SceneContext = FSceneRenderTargets::Get_Todo_PassContext();
+	FSceneRenderTargets& SceneContext = FSceneRenderTargets::Get(InRHICmdList);
 	if(SceneContext.IsSceneColorAllocated())
 	{
 		SceneColor = Graph.RegisterPass(new(FMemStack::Get()) FRCPassPostProcessInput(SceneContext.GetSceneColor()));
@@ -776,7 +776,7 @@ static void AddTemporalAA( FPostprocessContext& Context, FRenderingCompositeOutp
 	else
 	{
 		// No history so use current as history
-		HistoryInput = Context.Graph.RegisterPass( new(FMemStack::Get()) FRCPassPostProcessInput( FSceneRenderTargets::Get_Todo_PassContext().GetSceneColor() ) );
+		HistoryInput = Context.Graph.RegisterPass( new(FMemStack::Get()) FRCPassPostProcessInput( FSceneRenderTargets::Get(Context.RHICmdList).GetSceneColor() ) );
 	}
 
 	FRenderingCompositePass* TemporalAAPass = Context.Graph.RegisterPass( new(FMemStack::Get()) FRCPassPostProcessTemporalAA );
@@ -842,7 +842,7 @@ static FRenderingCompositePass* AddSinglePostProcessMaterial(FPostprocessContext
 		if(Material->NeedsGBuffer())
 		{
 			// AdjustGBufferRefCount(-1) call is done when the pass gets executed
-			FSceneRenderTargets::Get_Todo_PassContext().AdjustGBufferRefCount(Context.RHICmdList, 1);
+			FSceneRenderTargets::Get(Context.RHICmdList).AdjustGBufferRefCount(Context.RHICmdList, 1);
 		}
 
 		FRenderingCompositePass* Node = Context.Graph.RegisterPass(new(FMemStack::Get()) FRCPassPostProcessMaterial(MaterialInterface, Context.View.GetFeatureLevel()));
@@ -944,7 +944,7 @@ static void AddPostProcessMaterial(FPostprocessContext& Context, EBlendableLocat
 		if(Material->NeedsGBuffer())
 		{
 			// AdjustGBufferRefCount(-1) call is done when the pass gets executed
-			FSceneRenderTargets::Get_Todo_PassContext().AdjustGBufferRefCount(Context.RHICmdList, 1);
+			FSceneRenderTargets::Get(Context.RHICmdList).AdjustGBufferRefCount(Context.RHICmdList, 1);
 		}
 
 		FRenderingCompositePass* Node = Context.Graph.RegisterPass(new(FMemStack::Get()) FRCPassPostProcessMaterial(MaterialInterface,FeatureLevel));
@@ -1005,7 +1005,7 @@ static void AddHighResScreenshotMask(FPostprocessContext& Context, FRenderingCom
 		if (RendererMaterial->NeedsGBuffer())
 		{
 			// AdjustGBufferRefCount(-1) call is done when the pass gets executed
-			FSceneRenderTargets::Get_Todo_PassContext().AdjustGBufferRefCount(Context.RHICmdList, 1);
+			FSceneRenderTargets::Get(Context.RHICmdList).AdjustGBufferRefCount(Context.RHICmdList, 1);
 		}
 	}
 }
@@ -1056,7 +1056,7 @@ static void AddGBufferVisualizationOverview(FPostprocessContext& Context, FRende
 					if (Material->NeedsGBuffer())
 					{
 						// AdjustGBufferRefCount(-1) call is done when the pass gets executed
-						FSceneRenderTargets::Get_Todo_PassContext().AdjustGBufferRefCount(Context.RHICmdList, 1);
+						FSceneRenderTargets::Get(Context.RHICmdList).AdjustGBufferRefCount(Context.RHICmdList, 1);
 					}
 
 					if (BaseFilename.Len())
@@ -1905,7 +1905,7 @@ void FPostProcessing::Process(FRHICommandListImmediate& RHICmdList, const FViewI
 		{
 			// Generally we no longer need the GBuffers, anyone that wants to keep the GBuffers for longer should have called AdjustGBufferRefCount(1) to keep it for longer
 			// and call AdjustGBufferRefCount(-1) once it's consumed. This needs to happen each frame. PostProcessMaterial do that automatically
-			FSceneRenderTargets::Get_Todo_PassContext().AdjustGBufferRefCount(RHICmdList, -1);
+			FSceneRenderTargets::Get(RHICmdList).AdjustGBufferRefCount(RHICmdList, -1);
 		}
 
 		// The graph setup should be finished before this line ----------------------------------------
@@ -2416,7 +2416,7 @@ void FPostProcessing::ProcessPlanarReflection(FRHICommandListImmediate& RHICmdLi
 		FRenderingCompositePassContext CompositeContext(RHICmdList, View);
 
 		FPostprocessContext Context(RHICmdList, CompositeContext.Graph, View);
-		FSceneRenderTargets& SceneContext = FSceneRenderTargets::Get_Todo_PassContext();
+		FSceneRenderTargets& SceneContext = FSceneRenderTargets::Get(RHICmdList);
 
 		FRenderingCompositeOutputRef VelocityInput;
 		if(VelocityRT)
diff --git a/Engine/Source/Runtime/Renderer/Private/PostProcess/SceneRenderTargets.cpp b/Engine/Source/Runtime/Renderer/Private/PostProcess/SceneRenderTargets.cpp
--- a/Engine/Source/Runtime/Renderer/Private/PostProcess/SceneRenderTargets.cpp
+++ b/Engine/Source/Runtime/Renderer/Private/PostProcess/SceneRenderTargets.cpp
@@ -140,7 +140,7 @@ FSceneRenderTargets& FSceneRenderTargets::Get(FRHIAsyncComputeCommandListImmedia
 	return SceneRenderTargetsSingleton;
 }
 
-FSceneRenderTargets& FSceneRenderTargets::Get_Todo_PassContext()
+FSceneRenderTargets& FSceneRenderTargets::GetGlobalUnsafe()
 {
 	check(IsInRenderingThread()
 		&& !FTaskGraphInterface::Get().IsThreadProcessingTasks(ENamedThreads::RenderThread_Local)); // if we are processing tasks on the local queue, it is assumed this are in support of async tasks, which cannot use the current state of the render targets. This can be relaxed if needed.
@@ -808,8 +808,7 @@ void FSceneRenderTargets::GetGBufferADesc(FPooledRenderTargetDesc& Desc) const
 }
 
 void FSceneRenderTargets::AllocGBufferTargets(FRHICommandList& RHICmdList)
-{
-	check(!IsAnyForwardShadingEnabled(GetFeatureLevelShaderPlatform(CurrentFeatureLevel)));
+{	
 	// AdjustGBufferRefCount +1 doesn't match -1 (within the same frame)
 	ensure(GBufferRefCount == 0);
 
@@ -820,58 +819,64 @@ void FSceneRenderTargets::AllocGBufferTargets(FRHICommandList& RHICmdList)
 	}
 
 	// create GBuffer on demand so it can be shared with other pooled RT
+	const EShaderPlatform ShaderPlatform = GetFeatureLevelShaderPlatform(CurrentFeatureLevel);
+	const bool bUseGBuffer = IsUsingGBuffers(ShaderPlatform);
+	const bool bCanReadGBufferUniforms = (bUseGBuffer || IsSimpleForwardShadingEnabled(ShaderPlatform)) && CurrentFeatureLevel >= ERHIFeatureLevel::SM4;
+	if (bUseGBuffer)
+	{
+		// good to see the quality loss due to precision in the gbuffer
+		const bool bHighPrecisionGBuffers = (CurrentGBufferFormat >= EGBufferFormat::Force16BitsPerChannel);
+		// good to profile the impact of non 8 bit formats
+		const bool bEnforce8BitPerChannel = (CurrentGBufferFormat == EGBufferFormat::Force8BitsPerChannel);
 
-	// good to see the quality loss due to precision in the gbuffer
-	const bool bHighPrecisionGBuffers = (CurrentGBufferFormat >= EGBufferFormat::Force16BitsPerChannel);
-	// good to profile the impact of non 8 bit formats
-	const bool bEnforce8BitPerChannel = (CurrentGBufferFormat == EGBufferFormat::Force8BitsPerChannel);
+		// Create the world-space normal g-buffer.
+		{
+			FPooledRenderTargetDesc Desc;
+			GetGBufferADesc(Desc);
+			GRenderTargetPool.FindFreeElement(RHICmdList, Desc, GBufferA, TEXT("GBufferA"));
+		}
 
-	// Create the world-space normal g-buffer.
-	{
-		FPooledRenderTargetDesc Desc;
-		GetGBufferADesc(Desc);
-		GRenderTargetPool.FindFreeElement(RHICmdList, Desc, GBufferA, TEXT("GBufferA"));
-	}
+		// Create the specular color and power g-buffer.
+		{
+			const EPixelFormat SpecularGBufferFormat = bHighPrecisionGBuffers ? PF_FloatRGBA : PF_B8G8R8A8;
 
-	// Create the specular color and power g-buffer.
-	{
-		const EPixelFormat SpecularGBufferFormat = bHighPrecisionGBuffers ? PF_FloatRGBA : PF_B8G8R8A8;
+			FPooledRenderTargetDesc Desc(FPooledRenderTargetDesc::Create2DDesc(BufferSize, SpecularGBufferFormat, FClearValueBinding::Transparent, TexCreate_None, TexCreate_RenderTargetable, false));
+			GRenderTargetPool.FindFreeElement(RHICmdList, Desc, GBufferB, TEXT("GBufferB"));
+		}
 
-		FPooledRenderTargetDesc Desc(FPooledRenderTargetDesc::Create2DDesc(BufferSize, SpecularGBufferFormat, FClearValueBinding::Transparent, TexCreate_None, TexCreate_RenderTargetable, false));
-		GRenderTargetPool.FindFreeElement(RHICmdList, Desc, GBufferB, TEXT("GBufferB"));
-	}
+		// Create the diffuse color g-buffer.
+		{
+			const EPixelFormat DiffuseGBufferFormat = bHighPrecisionGBuffers ? PF_FloatRGBA : PF_B8G8R8A8;
+			FPooledRenderTargetDesc Desc(FPooledRenderTargetDesc::Create2DDesc(BufferSize, DiffuseGBufferFormat, FClearValueBinding::Transparent, TexCreate_SRGB, TexCreate_RenderTargetable, false));
+			GRenderTargetPool.FindFreeElement(RHICmdList, Desc, GBufferC, TEXT("GBufferC"));
+		}
 
-	// Create the diffuse color g-buffer.
-	{
-		const EPixelFormat DiffuseGBufferFormat = bHighPrecisionGBuffers ? PF_FloatRGBA : PF_B8G8R8A8;
-		FPooledRenderTargetDesc Desc(FPooledRenderTargetDesc::Create2DDesc(BufferSize, DiffuseGBufferFormat, FClearValueBinding::Transparent, TexCreate_SRGB, TexCreate_RenderTargetable, false));
-		GRenderTargetPool.FindFreeElement(RHICmdList, Desc, GBufferC, TEXT("GBufferC"));
-	}
+		// Create the mask g-buffer (e.g. SSAO, subsurface scattering, wet surface mask, skylight mask, ...).
+		{
+			FPooledRenderTargetDesc Desc(FPooledRenderTargetDesc::Create2DDesc(BufferSize, PF_B8G8R8A8, FClearValueBinding(FLinearColor(0, 1, 1, 1)), TexCreate_None, TexCreate_RenderTargetable, false));
+			GRenderTargetPool.FindFreeElement(RHICmdList, Desc, GBufferD, TEXT("GBufferD"));
+		}
 
-	// Create the mask g-buffer (e.g. SSAO, subsurface scattering, wet surface mask, skylight mask, ...).
-	{
-		FPooledRenderTargetDesc Desc(FPooledRenderTargetDesc::Create2DDesc(BufferSize, PF_B8G8R8A8, FClearValueBinding(FLinearColor(0, 1, 1, 1)), TexCreate_None, TexCreate_RenderTargetable, false));
-		GRenderTargetPool.FindFreeElement(RHICmdList, Desc, GBufferD, TEXT("GBufferD"));
-	}
+		if (bAllowStaticLighting)
+		{
+			FPooledRenderTargetDesc Desc(FPooledRenderTargetDesc::Create2DDesc(BufferSize, PF_B8G8R8A8, FClearValueBinding(FLinearColor(1, 1, 1, 1)), TexCreate_None, TexCreate_RenderTargetable, false));
+			GRenderTargetPool.FindFreeElement(RHICmdList, Desc, GBufferE, TEXT("GBufferE"));
+		}
 
-	if (bAllowStaticLighting)
-	{
-		FPooledRenderTargetDesc Desc(FPooledRenderTargetDesc::Create2DDesc(BufferSize, PF_B8G8R8A8, FClearValueBinding(FLinearColor(1, 1, 1, 1)), TexCreate_None, TexCreate_RenderTargetable, false));
-		GRenderTargetPool.FindFreeElement(RHICmdList, Desc, GBufferE, TEXT("GBufferE"));
-	}
+		GBasePassOutputsVelocityDebug = CVarBasePassOutputsVelocityDebug.GetValueOnRenderThread();
+		if (bAllocateVelocityGBuffer)
+		{
+			FPooledRenderTargetDesc VelocityRTDesc = FVelocityRendering::GetRenderTargetDesc();
+			GRenderTargetPool.FindFreeElement(RHICmdList, VelocityRTDesc, GBufferVelocity, TEXT("GBufferVelocity"));
+		}
 
-	GBasePassOutputsVelocityDebug = CVarBasePassOutputsVelocityDebug.GetValueOnRenderThread();
-	if (bAllocateVelocityGBuffer)
-	{
-		FPooledRenderTargetDesc VelocityRTDesc = FVelocityRendering::GetRenderTargetDesc();
-		GRenderTargetPool.FindFreeElement(RHICmdList, VelocityRTDesc, GBufferVelocity, TEXT("GBufferVelocity"));
+		// otherwise we have a severe problem
+		check(GBufferA);
 	}
 
-	// otherwise we have a severe problem
-	check(GBufferA);
-
 	// Create the required render targets if running Highend.
-	if (CurrentFeatureLevel >= ERHIFeatureLevel::SM4)
+	// Allow creation of dummy resources in case of simpleforward lighting.
+	if (bCanReadGBufferUniforms)
 	{
 
 		// Allocate the Gbuffer resource uniform buffer.
@@ -1015,23 +1020,20 @@ TRefCountPtr<IPooledRenderTarget>& FSceneRenderTargets::GetLightAttenuation()
 }
 
 void FSceneRenderTargets::AdjustGBufferRefCount(FRHICommandList& RHICmdList, int Delta)
-{ 
-	if (!IsAnyForwardShadingEnabled(GetFeatureLevelShaderPlatform(CurrentFeatureLevel)))
+{
+	if (Delta > 0 && GBufferRefCount == 0)
 	{
-		if (Delta > 0 && GBufferRefCount == 0)
-		{
-			AllocGBufferTargets(RHICmdList);
-		}
-		else
-		{
-			GBufferRefCount += Delta;
+		AllocGBufferTargets(RHICmdList);
+	}
+	else
+	{
+		GBufferRefCount += Delta;
 
-			if (GBufferRefCount == 0)
-			{
-				ReleaseGBufferTargets();
-			}
+		if (GBufferRefCount == 0)
+		{
+			ReleaseGBufferTargets();
 		}
-	}
+	}	
 }
 
 void FSceneRenderTargets::FinishRenderingSceneColor(FRHICommandListImmediate& RHICmdList, const FResolveRect& ResolveRect)
diff --git a/Engine/Source/Runtime/Renderer/Private/PostProcess/SceneRenderTargets.h b/Engine/Source/Runtime/Renderer/Private/PostProcess/SceneRenderTargets.h
--- a/Engine/Source/Runtime/Renderer/Private/PostProcess/SceneRenderTargets.h
+++ b/Engine/Source/Runtime/Renderer/Private/PostProcess/SceneRenderTargets.h
@@ -154,7 +154,7 @@ public:
 	static FSceneRenderTargets& Get(FRHIAsyncComputeCommandListImmediate& RHICmdList);
 
 	// this is a placeholder, the context should come from somewhere. This is very unsafe, please don't use it!
-	static FSceneRenderTargets& Get_Todo_PassContext();
+	static FSceneRenderTargets& GetGlobalUnsafe();
 	// As above but relaxed checks and always gives the global FSceneRenderTargets. The intention here is that it is only used for constants that don't change during a frame. This is very unsafe, please don't use it!
 	static FSceneRenderTargets& Get_FrameConstantsOnly();
 
diff --git a/Engine/Source/Runtime/Renderer/Private/Renderer.cpp b/Engine/Source/Runtime/Renderer/Private/Renderer.cpp
--- a/Engine/Source/Runtime/Renderer/Private/Renderer.cpp
+++ b/Engine/Source/Runtime/Renderer/Private/Renderer.cpp
@@ -23,13 +23,13 @@ IMPLEMENT_MODULE(FRendererModule, Renderer);
 void FRendererModule::ReallocateSceneRenderTargets()
 {
 	FLightPrimitiveInteraction::InitializeMemoryPool();
-	FSceneRenderTargets::Get_Todo_PassContext().UpdateRHI();
+	FSceneRenderTargets::GetGlobalUnsafe().UpdateRHI();
 }
 
 void FRendererModule::SceneRenderTargetsSetBufferSize(uint32 SizeX, uint32 SizeY)
 {
-	FSceneRenderTargets::Get_Todo_PassContext().SetBufferSize(SizeX, SizeY);
-	FSceneRenderTargets::Get_Todo_PassContext().UpdateRHI();
+	FSceneRenderTargets::GetGlobalUnsafe().SetBufferSize(SizeX, SizeY);
+	FSceneRenderTargets::GetGlobalUnsafe().UpdateRHI();
 }
 
 void FRendererModule::InitializeSystemTextures(FRHICommandListImmediate& RHICmdList)
diff --git a/Engine/Source/Runtime/Renderer/Private/RendererScene.cpp b/Engine/Source/Runtime/Renderer/Private/RendererScene.cpp
--- a/Engine/Source/Runtime/Renderer/Private/RendererScene.cpp
+++ b/Engine/Source/Runtime/Renderer/Private/RendererScene.cpp
@@ -1723,6 +1723,8 @@ void FScene::AddWindSource(UWindDirectionalSourceComponent* WindComponent)
 		return;
 	}
 
+	WindComponents_GameThread.Add(WindComponent);
+
 	FWindSourceSceneProxy* SceneProxy = WindComponent->CreateSceneProxy();
 	WindComponent->SceneProxy = SceneProxy;
 
@@ -1737,6 +1739,8 @@ void FScene::AddWindSource(UWindDirectionalSourceComponent* WindComponent)
 
 void FScene::RemoveWindSource(UWindDirectionalSourceComponent* WindComponent)
 {
+	WindComponents_GameThread.Remove(WindComponent);
+
 	FWindSourceSceneProxy* SceneProxy = WindComponent->SceneProxy;
 	WindComponent->SceneProxy = NULL;
 
@@ -1762,7 +1766,7 @@ const TArray<FWindSourceSceneProxy*>& FScene::GetWindSources_RenderThread() cons
 
 void FScene::GetWindParameters(const FVector& Position, FVector& OutDirection, float& OutSpeed, float& OutMinGustAmt, float& OutMaxGustAmt) const
 {
-	FWindSourceSceneProxy::FWindData AccumWindData;
+	FWindData AccumWindData;
 	AccumWindData.PrepareForAccumulate();
 
 	int32 NumActiveWindSources = 0;
@@ -1774,7 +1778,7 @@ void FScene::GetWindParameters(const FVector& Position, FVector& OutDirection, f
 		FVector4 CurrentDirectionAndSpeed;
 		float Weight;
 		const FWindSourceSceneProxy* CurrentSource = WindSources[i];
-		FWindSourceSceneProxy::FWindData CurrentSourceData;
+		FWindData CurrentSourceData;
 		if (CurrentSource->GetWindParameters(Position, CurrentSourceData, Weight))
 		{
 			AccumWindData.AddWeighted(CurrentSourceData, Weight);
@@ -1795,9 +1799,44 @@ void FScene::GetWindParameters(const FVector& Position, FVector& OutDirection, f
 	OutMaxGustAmt	= AccumWindData.MaxGustAmt;
 }
 
+void FScene::GetWindParameters_GameThread(const FVector& Position, FVector& OutDirection, float& OutSpeed, float& OutMinGustAmt, float& OutMaxGustAmt) const
+{
+	FWindData AccumWindData;
+	AccumWindData.PrepareForAccumulate();
+
+	const int32 NumSources = WindComponents_GameThread.Num();
+	int32 NumActiveSources = 0;
+	float TotalWeight = 0.0f;
+
+	// read the wind component array, this is safe for the game thread
+	for(UWindDirectionalSourceComponent* Component : WindComponents_GameThread)
+	{
+		float Weight = 0.0f;
+		FWindData CurrentComponentData;
+		if(Component->GetWindParameters(Position, CurrentComponentData, Weight))
+		{
+			AccumWindData.AddWeighted(CurrentComponentData, Weight);
+			TotalWeight += Weight;
+			++NumActiveSources;
+		}
+	}
+
+	AccumWindData.NormalizeByTotalWeight(TotalWeight);
+
+	if(NumActiveSources == 0)
+	{
+		AccumWindData.Direction = FVector(1.0f, 0.0f, 0.0f);
+	}
+
+	OutDirection = AccumWindData.Direction;
+	OutSpeed = AccumWindData.Speed;
+	OutMinGustAmt = AccumWindData.MinGustAmt;
+	OutMaxGustAmt = AccumWindData.MaxGustAmt;
+}
+
 void FScene::GetDirectionalWindParameters(FVector& OutDirection, float& OutSpeed, float& OutMinGustAmt, float& OutMaxGustAmt) const
 {
-	FWindSourceSceneProxy::FWindData AccumWindData;
+	FWindData AccumWindData;
 	AccumWindData.PrepareForAccumulate();
 
 	int32 NumActiveWindSources = 0;
@@ -1808,7 +1847,7 @@ void FScene::GetDirectionalWindParameters(FVector& OutDirection, float& OutSpeed
 		FVector4 CurrentDirectionAndSpeed;
 		float Weight;
 		const FWindSourceSceneProxy* CurrentSource = WindSources[i];
-		FWindSourceSceneProxy::FWindData CurrentSourceData;
+		FWindData CurrentSourceData;
 		if (CurrentSource->GetDirectionalWindParameters(CurrentSourceData, Weight))
 		{
 			AccumWindData.AddWeighted(CurrentSourceData, Weight);			
@@ -2656,6 +2695,7 @@ public:
 		return NullWindSources;
 	}
 	virtual void GetWindParameters(const FVector& Position, FVector& OutDirection, float& OutSpeed, float& OutMinGustAmt, float& OutMaxGustAmt) const override { OutDirection = FVector(1.0f, 0.0f, 0.0f); OutSpeed = 0.0f; OutMinGustAmt = 0.0f; OutMaxGustAmt = 0.0f; }
+	virtual void GetWindParameters_GameThread(const FVector& Position, FVector& OutDirection, float& OutSpeed, float& OutMinGustAmt, float& OutMaxGustAmt) const override { OutDirection = FVector(1.0f, 0.0f, 0.0f); OutSpeed = 0.0f; OutMinGustAmt = 0.0f; OutMaxGustAmt = 0.0f; }
 	virtual void GetDirectionalWindParameters(FVector& OutDirection, float& OutSpeed, float& OutMinGustAmt, float& OutMaxGustAmt) const override { OutDirection = FVector(1.0f, 0.0f, 0.0f); OutSpeed = 0.0f; OutMinGustAmt = 0.0f; OutMaxGustAmt = 0.0f; }
 	virtual void AddSpeedTreeWind(class FVertexFactory* VertexFactory, const class UStaticMesh* StaticMesh) override {}
 	virtual void RemoveSpeedTreeWind(class FVertexFactory* VertexFactory, const class UStaticMesh* StaticMesh) override {}
diff --git a/Engine/Source/Runtime/Renderer/Private/ScenePrivate.h b/Engine/Source/Runtime/Renderer/Private/ScenePrivate.h
--- a/Engine/Source/Runtime/Renderer/Private/ScenePrivate.h
+++ b/Engine/Source/Runtime/Renderer/Private/ScenePrivate.h
@@ -1909,6 +1909,9 @@ public:
 	/** The wind sources in the scene. */
 	TArray<class FWindSourceSceneProxy*> WindSources;
 
+	/** Wind source components, tracked so the game thread can also access wind parameters */
+	TArray<UWindDirectionalSourceComponent*> WindComponents_GameThread;
+
 	/** SpeedTree wind objects in the scene. FLocalVertexFactoryShaderParameters needs to lookup by FVertexFactory, but wind objects are per tree (i.e. per UStaticMesh)*/
 	TMap<const UStaticMesh*, struct FSpeedTreeWindComputation*> SpeedTreeWindComputationMap;
 	TMap<FVertexFactory*, const UStaticMesh*> SpeedTreeVertexFactoryMap;
@@ -2003,6 +2006,7 @@ public:
 	virtual void RemoveWindSource(UWindDirectionalSourceComponent* WindComponent) override;
 	virtual const TArray<FWindSourceSceneProxy*>& GetWindSources_RenderThread() const override;
 	virtual void GetWindParameters(const FVector& Position, FVector& OutDirection, float& OutSpeed, float& OutMinGustAmt, float& OutMaxGustAmt) const override;
+	virtual void GetWindParameters_GameThread(const FVector& Position, FVector& OutDirection, float& OutSpeed, float& OutMinGustAmt, float& OutMaxGustAmt) const override;
 	virtual void GetDirectionalWindParameters(FVector& OutDirection, float& OutSpeed, float& OutMinGustAmt, float& OutMaxGustAmt) const override;
 	virtual void AddSpeedTreeWind(FVertexFactory* VertexFactory, const UStaticMesh* StaticMesh) override;
 	virtual void RemoveSpeedTreeWind(FVertexFactory* VertexFactory, const UStaticMesh* StaticMesh) override;
diff --git a/Engine/Source/Runtime/Renderer/Private/SceneRendering.cpp b/Engine/Source/Runtime/Renderer/Private/SceneRendering.cpp
--- a/Engine/Source/Runtime/Renderer/Private/SceneRendering.cpp
+++ b/Engine/Source/Runtime/Renderer/Private/SceneRendering.cpp
@@ -1478,9 +1478,9 @@ void FSceneRenderer::RenderCustomDepthPass(FRHICommandListImmediate& RHICmdList)
 	}
 }
 
-void FSceneRenderer::OnStartFrame()
+void FSceneRenderer::OnStartFrame(FRHICommandListImmediate& RHICmdList)
 {
-	FSceneRenderTargets& SceneContext = FSceneRenderTargets::Get_Todo_PassContext();
+	FSceneRenderTargets& SceneContext = FSceneRenderTargets::Get(RHICmdList);
 
 	GRenderTargetPool.VisualizeTexture.OnStartFrame(Views[0]);
 	CompositionGraph_OnStartFrame();
diff --git a/Engine/Source/Runtime/Renderer/Private/SceneRendering.h b/Engine/Source/Runtime/Renderer/Private/SceneRendering.h
--- a/Engine/Source/Runtime/Renderer/Private/SceneRendering.h
+++ b/Engine/Source/Runtime/Renderer/Private/SceneRendering.h
@@ -1336,7 +1336,7 @@ protected:
 	void RenderCustomDepthPassAtLocation(FRHICommandListImmediate& RHICmdList, int32 Location);
 	void RenderCustomDepthPass(FRHICommandListImmediate& RHICmdList);
 
-	void OnStartFrame();
+	void OnStartFrame(FRHICommandListImmediate& RHICmdList);
 
 	/** Renders the scene's distortion */
 	void RenderDistortion(FRHICommandListImmediate& RHICmdList);
diff --git a/Engine/Source/Runtime/Renderer/Private/SceneVisibility.cpp b/Engine/Source/Runtime/Renderer/Private/SceneVisibility.cpp
--- a/Engine/Source/Runtime/Renderer/Private/SceneVisibility.cpp
+++ b/Engine/Source/Runtime/Renderer/Private/SceneVisibility.cpp
@@ -2868,7 +2868,7 @@ bool FDeferredShadingSceneRenderer::InitViews(FRHICommandListImmediate& RHICmdLi
 
 	{
 		QUICK_SCOPE_CYCLE_COUNTER(STAT_InitViews_OnStartFrame);
-		OnStartFrame();
+		OnStartFrame(RHICmdList);
 	}
 
 	return bDoInitViewAftersPrepass;
diff --git a/Engine/Source/Runtime/Slate/Private/Framework/Application/SlateApplication.cpp b/Engine/Source/Runtime/Slate/Private/Framework/Application/SlateApplication.cpp
--- a/Engine/Source/Runtime/Slate/Private/Framework/Application/SlateApplication.cpp
+++ b/Engine/Source/Runtime/Slate/Private/Framework/Application/SlateApplication.cpp
@@ -5898,7 +5898,7 @@ bool FSlateApplication::OnSizeChanged( const TSharedRef< FGenericWindow >& Platf
 
 		Renderer->RequestResize( Window, Width, Height );
 
-		if ( !bWasMinimized && Window->IsRegularWindow() && !Window->HasOSWindowBorder() && Window->IsVisible() )
+		if ( !bWasMinimized && Window->IsRegularWindow() && !Window->HasOSWindowBorder() && Window->IsVisible() && Window->IsDrawingEnabled() )
 		{
 			PrivateDrawWindows( Window );
 		}
diff --git a/Engine/Source/Runtime/Slate/Private/Framework/Text/TextLayout.cpp b/Engine/Source/Runtime/Slate/Private/Framework/Text/TextLayout.cpp
--- a/Engine/Source/Runtime/Slate/Private/Framework/Text/TextLayout.cpp
+++ b/Engine/Source/Runtime/Slate/Private/Framework/Text/TextLayout.cpp
@@ -2158,6 +2158,11 @@ bool FTextLayout::IsEmpty() const
 	return (LineModels.Num() == 0 || (LineModels.Num() == 1 && LineModels[0].Text->Len() == 0));
 }
 
+int32 FTextLayout::GetLineCount() const
+{
+	return LineModels.Num();
+}
+
 void FTextLayout::GetAsText(FString& DisplayText, FTextOffsetLocations* const OutTextOffsetLocations) const
 {
 	GetAsTextAndOffsets(&DisplayText, OutTextOffsetLocations);
diff --git a/Engine/Source/Runtime/Slate/Private/Widgets/Input/SInputKeySelector.cpp b/Engine/Source/Runtime/Slate/Private/Widgets/Input/SInputKeySelector.cpp
--- a/Engine/Source/Runtime/Slate/Private/Widgets/Input/SInputKeySelector.cpp
+++ b/Engine/Source/Runtime/Slate/Private/Widgets/Input/SInputKeySelector.cpp
@@ -42,11 +42,15 @@ FText SInputKeySelector::GetSelectedKeyText() const
 	}
 	else if ( SelectedKey.IsSet() )
 	{
-		// If the key in the chord is a modifier key, print it's display name directly since the FInputChord
-		// displays these as empty text.
-		return SelectedKey.Get().Key.IsModifierKey()
-			? SelectedKey.Get().Key.GetDisplayName()
-			: SelectedKey.Get().GetInputText();
+		if(SelectedKey.Get().Key.IsValid())
+		{
+			// If the key in the chord is a modifier key, print it's display name directly since the FInputChord
+					// displays these as empty text.
+			return SelectedKey.Get().Key.IsModifierKey()
+				? SelectedKey.Get().Key.GetDisplayName()
+				: SelectedKey.Get().GetInputText();
+		}
+		return FText();
 	}
 	return FText();
 }
diff --git a/Engine/Source/Runtime/Slate/Private/Widgets/Layout/SSafeZone.cpp b/Engine/Source/Runtime/Slate/Private/Widgets/Layout/SSafeZone.cpp
--- a/Engine/Source/Runtime/Slate/Private/Widgets/Layout/SSafeZone.cpp
+++ b/Engine/Source/Runtime/Slate/Private/Widgets/Layout/SSafeZone.cpp
@@ -24,6 +24,13 @@ void SSafeZone::Construct( const FArguments& InArgs )
 	bPadBottom = InArgs._PadBottom;
 
 	SetTitleSafe(bIsTitleSafe);
+
+	FCoreDelegates::OnSafeFrameChangedEvent.AddSP(this, &SSafeZone::SafeAreaUpdated);
+}
+
+void SSafeZone::SafeAreaUpdated()
+{
+	SetTitleSafe(bIsTitleSafe);
 }
 
 void SSafeZone::SetTitleSafe( bool InIsTitleSafe )
diff --git a/Engine/Source/Runtime/Slate/Private/Widgets/Layout/SScaleBox.cpp b/Engine/Source/Runtime/Slate/Private/Widgets/Layout/SScaleBox.cpp
--- a/Engine/Source/Runtime/Slate/Private/Widgets/Layout/SScaleBox.cpp
+++ b/Engine/Source/Runtime/Slate/Private/Widgets/Layout/SScaleBox.cpp
@@ -24,6 +24,8 @@ void SScaleBox::Construct( const SScaleBox::FArguments& InArgs )
 	[
 		InArgs._Content.Widget
 	];
+
+	FCoreDelegates::OnSafeFrameChangedEvent.AddSP(this, &SScaleBox::RefreshSafeZoneScale);
 }
 
 /* SWidget overrides
diff --git a/Engine/Source/Runtime/Slate/Public/Framework/Commands/InputChord.h b/Engine/Source/Runtime/Slate/Public/Framework/Commands/InputChord.h
--- a/Engine/Source/Runtime/Slate/Public/Framework/Commands/InputChord.h
+++ b/Engine/Source/Runtime/Slate/Public/Framework/Commands/InputChord.h
@@ -14,23 +14,23 @@ struct SLATE_API FInputChord
 	GENERATED_USTRUCT_BODY()
 
 	/** The Key is the core of the chord. */
-	UPROPERTY()
+	UPROPERTY(EditAnywhere, Category = Key)
 	FKey Key;
 
 	/** Whether the shift key is part of the chord.  */
-	UPROPERTY()
+	UPROPERTY(EditAnywhere, Category = Modifier)
 	uint32 bShift:1;
 
 	/** Whether the control key is part of the chord.  */
-	UPROPERTY()
+	UPROPERTY(EditAnywhere, Category = Modifier)
 	uint32 bCtrl:1;
 
 	/** Whether the alt key is part of the chord.  */
-	UPROPERTY()
+	UPROPERTY(EditAnywhere, Category = Modifier)
 	uint32 bAlt:1;
 
 	/** Whether the command key is part of the chord.  */
-	UPROPERTY()
+	UPROPERTY(EditAnywhere, Category = Modifier)
 	uint32 bCmd:1;
 
 	/** 
diff --git a/Engine/Source/Runtime/Slate/Public/Framework/Text/TextLayout.h b/Engine/Source/Runtime/Slate/Public/Framework/Text/TextLayout.h
--- a/Engine/Source/Runtime/Slate/Public/Framework/Text/TextLayout.h
+++ b/Engine/Source/Runtime/Slate/Public/Framework/Text/TextLayout.h
@@ -485,6 +485,8 @@ public:
 
 	bool IsEmpty() const;
 
+	int32 GetLineCount() const;
+
 	void GetAsText(FString& DisplayText, FTextOffsetLocations* const OutTextOffsetLocations = nullptr) const;
 
 	void GetAsText(FText& DisplayText, FTextOffsetLocations* const OutTextOffsetLocations = nullptr) const;
diff --git a/Engine/Source/Runtime/Slate/Public/Widgets/Input/SInputKeySelector.h b/Engine/Source/Runtime/Slate/Public/Widgets/Input/SInputKeySelector.h
--- a/Engine/Source/Runtime/Slate/Public/Widgets/Input/SInputKeySelector.h
+++ b/Engine/Source/Runtime/Slate/Public/Widgets/Input/SInputKeySelector.h
@@ -8,7 +8,7 @@
 class SLATE_API SInputKeySelector : public SCompoundWidget
 {
 public:
-	DECLARE_DELEGATE_OneParam( FOnKeySelected, FInputChord )
+	DECLARE_DELEGATE_OneParam( FOnKeySelected, const FInputChord& )
 	DECLARE_DELEGATE( FOnIsSelectingKeyChanged )
 
 	SLATE_BEGIN_ARGS( SInputKeySelector )
diff --git a/Engine/Source/Runtime/Slate/Public/Widgets/Layout/SConstraintCanvas.h b/Engine/Source/Runtime/Slate/Public/Widgets/Layout/SConstraintCanvas.h
--- a/Engine/Source/Runtime/Slate/Public/Widgets/Layout/SConstraintCanvas.h
+++ b/Engine/Source/Runtime/Slate/Public/Widgets/Layout/SConstraintCanvas.h
@@ -130,7 +130,7 @@ public:
 	void ClearChildren();
 
 	/** Get the cached geometry of the canvas. */
-	FGeometry GetCachedGeometry() const { return CachedGeometry; }
+	const FGeometry& GetCachedGeometry() const { return CachedGeometry; }
 
 public:
 
diff --git a/Engine/Source/Runtime/Slate/Public/Widgets/Layout/SSafeZone.h b/Engine/Source/Runtime/Slate/Public/Widgets/Layout/SSafeZone.h
--- a/Engine/Source/Runtime/Slate/Public/Widgets/Layout/SSafeZone.h
+++ b/Engine/Source/Runtime/Slate/Public/Widgets/Layout/SSafeZone.h
@@ -72,7 +72,8 @@ class SLATE_API SSafeZone : public SBox
 public:
 
 	void Construct( const FArguments& InArgs );
-
+	
+	void SafeAreaUpdated();
 	void SetTitleSafe( bool bIsTitleSafe );
 	void SetSafeAreaScale(FMargin InSafeAreaScale);
 
diff --git a/Engine/Source/Runtime/SlateCore/Private/Rendering/DrawElements.cpp b/Engine/Source/Runtime/SlateCore/Private/Rendering/DrawElements.cpp
--- a/Engine/Source/Runtime/SlateCore/Private/Rendering/DrawElements.cpp
+++ b/Engine/Source/Runtime/SlateCore/Private/Rendering/DrawElements.cpp
@@ -288,6 +288,17 @@ void FSlateDrawElement::MakeLayer(FSlateWindowElementList& ElementList, uint32 I
 	DrawElt.DataPayload.SetLayerPayloadProperties(DrawLayerHandle.Get());
 }
 
+void FSlateDrawElement::MakePostProcessPass(FSlateWindowElementList& ElementList, uint32 InLayer, const FPaintGeometry& PaintGeometry, const FSlateRect& InClippingRect, const FVector4& Params, int32 DownsampleAmount)
+{
+	SCOPE_CYCLE_COUNTER(STAT_SlateDrawElementMakeTime)
+
+	FSlateDrawElement& DrawElt = ElementList.AddUninitialized();
+	DrawElt.Init(InLayer, PaintGeometry, InClippingRect, ESlateDrawEffect::None);
+	DrawElt.ElementType = ET_PostProcessPass;
+	DrawElt.DataPayload.DownsampleAmount = DownsampleAmount;
+	DrawElt.DataPayload.PostProcessData = Params;
+}
+
 FVector2D FSlateDrawElement::GetRotationPoint(const FPaintGeometry& PaintGeometry, const TOptional<FVector2D>& UserRotationPoint, ERotationSpace RotationSpace)
 {
 	FVector2D RotationPoint(0, 0);
@@ -296,17 +307,17 @@ FVector2D FSlateDrawElement::GetRotationPoint(const FPaintGeometry& PaintGeometr
 
 	switch (RotationSpace)
 	{
-	case RelativeToElement:
-	{
-		// If the user did not specify a rotation point, we rotate about the center of the element
-		RotationPoint = UserRotationPoint.Get(LocalSize * 0.5f);
-	}
+		case RelativeToElement:
+		{
+			// If the user did not specify a rotation point, we rotate about the center of the element
+			RotationPoint = UserRotationPoint.Get(LocalSize * 0.5f);
+		}
 		break;
-	case RelativeToWorld:
-	{
-		// its in world space, must convert the point to local space.
-		RotationPoint = TransformPoint(Inverse(PaintGeometry.GetAccumulatedRenderTransform()), UserRotationPoint.Get(FVector2D::ZeroVector));
-	}
+		case RelativeToWorld:
+		{
+			// its in world space, must convert the point to local space.
+			RotationPoint = TransformPoint(Inverse(PaintGeometry.GetAccumulatedRenderTransform()), UserRotationPoint.Get(FVector2D::ZeroVector));
+		}
 		break;
 	default:
 		check(0);
@@ -474,6 +485,10 @@ void FSlateBatchData::Merge(FElementBatchMap& InLayerToElementBatches, uint32& V
 			{
 				AddRenderBatch(Layer, ElementBatch, 0, 0, 0, 0);
 			}
+			else if(ElementBatch.GetShaderType() == ESlateShader::PostProcess)
+			{
+				AddRenderBatch(Layer,ElementBatch, 0, 0, 0, 0);
+			}
 			else
 			{
 				if ( bExpandLayersAndCachedHandles )
diff --git a/Engine/Source/Runtime/SlateCore/Private/Rendering/ElementBatcher.cpp b/Engine/Source/Runtime/SlateCore/Private/Rendering/ElementBatcher.cpp
--- a/Engine/Source/Runtime/SlateCore/Private/Rendering/ElementBatcher.cpp
+++ b/Engine/Source/Runtime/SlateCore/Private/Rendering/ElementBatcher.cpp
@@ -21,16 +21,23 @@ DECLARE_DWORD_COUNTER_STAT(TEXT("Elements (Text)"), STAT_SlateNumTextElements, S
 // Super-hacky way of storing the scissor rect so we don't have to change all the FSlateDrawElement APIs for this hacky support.
 SLATECORE_API TOptional<FShortRect> GSlateScissorRect;
 
-
-
+namespace ElementBatcherUtils
+{
+	static FVector2D RoundToInt(const FVector2D& Vec)
+	{
+		return FVector2D(FMath::RoundToInt(Vec.X), FMath::RoundToInt(Vec.Y));
+	}
+}
 
 FSlateElementBatcher::FSlateElementBatcher( TSharedRef<FSlateRenderingPolicy> InRenderingPolicy )
 	: BatchData( nullptr )
 	, DrawLayer( nullptr )
 	, RenderingPolicy( &InRenderingPolicy.Get() )
 	, NumDrawnBatchesStat(0)
+	, NumPostProcessPasses(0)
 	, PixelCenterOffset( InRenderingPolicy->GetPixelCenterOffset() )
 	, bSRGBVertexColor( !InRenderingPolicy->IsVertexColorInLinearSpace() )
+	, bRequiresVsync(false)
 {
 }
 
@@ -49,12 +56,12 @@ void FSlateElementBatcher::AddElements(FSlateWindowElementList& WindowElementLis
 
 	BatchData = &WindowElementList.GetBatchData();
 
-	AddElements(WindowElementList.GetRootDrawLayer());
+	AddElements(WindowElementList, WindowElementList.GetRootDrawLayer());
 
 	TMap < TSharedPtr<FSlateDrawLayerHandle, ESPMode::ThreadSafe>, TSharedPtr<FSlateDrawLayer> >& DrawLayers = WindowElementList.GetChildDrawLayers();
 	for ( auto& Entry : DrawLayers )
 	{
-		AddElements(*Entry.Value.Get());
+		AddElements(WindowElementList, *Entry.Value.Get());
 	}
 
 	// Done with the element list
@@ -67,7 +74,7 @@ void FSlateElementBatcher::AddElements(FSlateWindowElementList& WindowElementLis
 	FPlatformMisc::EndNamedEvent();
 }
 
-void FSlateElementBatcher::AddElements(FSlateDrawLayer& InDrawLayer)
+void FSlateElementBatcher::AddElements(const FSlateWindowElementList& ElementList, FSlateDrawLayer& InDrawLayer)
 {
 	// This stuff is just for the counters. Could be scoped by an #ifdef if necessary.
 	static_assert(
@@ -84,7 +91,8 @@ void FSlateElementBatcher::AddElements(FSlateDrawLayer& InDrawLayer)
 		FSlateDrawElement::EElementType::ET_CustomVerts == 10 &&
 		FSlateDrawElement::EElementType::ET_CachedBuffer == 11 &&
 		FSlateDrawElement::EElementType::ET_Layer == 12 &&
-		FSlateDrawElement::EElementType::ET_Count == 13,
+		FSlateDrawElement::EElementType::ET_PostProcessPass == 13 &&
+		FSlateDrawElement::EElementType::ET_Count == 14,
 		"If FSlateDrawElement::EElementType is modified, this array must be made to match." );
 
 	static FName ElementFNames[] =
@@ -102,6 +110,7 @@ void FSlateElementBatcher::AddElements(FSlateDrawLayer& InDrawLayer)
 		FName(TEXT("CustomVerts")),
 		FName(TEXT("CachedBuffer")),
 		FName(TEXT("Layer")),
+		FName(TEXT("FXPass")),
 	};
 
 	DrawLayer = &InDrawLayer;
@@ -175,13 +184,16 @@ void FSlateElementBatcher::AddElements(FSlateDrawLayer& InDrawLayer)
 					AddCustomElement( DrawElement );
 					break;
 				case FSlateDrawElement::ET_CustomVerts:
-					AddCustomVerts(DrawElement);
+					AddCustomVerts( DrawElement);
 					break;
 				case FSlateDrawElement::ET_CachedBuffer:
 					AddCachedBuffer( DrawElement );
 					break;
 				case FSlateDrawElement::ET_Layer:
-					AddLayer(DrawElement);
+					AddLayer( DrawElement );
+					break;
+				case FSlateDrawElement::ET_PostProcessPass:
+					AddPostProcessPass( DrawElement, ElementList.GetWindow()->GetViewportSize());
 					break;
 				default:
 					checkf(0, TEXT("Invalid element type"));
@@ -1481,7 +1493,7 @@ void FSlateElementBatcher::AddViewportElement( const FSlateDrawElement& DrawElem
 
 	ESlateBatchDrawFlag::Type DrawFlags = InPayload.BatchFlags;
 
-	FSlateShaderResource* ViewportResource = InPayload.ViewportRenderTargetTexture;
+	FSlateShaderResource* ViewportResource = InPayload.RenderTargetResource;
 	ESlateShader::Type ShaderType = ESlateShader::Default;
 
 	if( InPayload.bViewportTextureAlphaOnly )
@@ -1870,7 +1882,64 @@ void FSlateElementBatcher::AddLayer(const FSlateDrawElement& DrawElement)
 }
 
 
-FSlateElementBatch& FSlateElementBatcher::FindBatchForElement( 
+void FSlateElementBatcher::AddPostProcessPass(const FSlateDrawElement& DrawElement, const FVector2D& WindowSize)
+{
+	++NumPostProcessPasses;
+
+	const FSlateRenderTransform& RenderTransform = DrawElement.GetRenderTransform();
+	const FVector2D& LocalSize = DrawElement.GetLocalSize();
+	
+	const FSlateDataPayload& Payload = DrawElement.GetDataPayload();
+
+	//@todo Doesnt work with rotated or skewed objects yet
+	const FVector2D& Position = DrawElement.GetPosition();
+
+	const FSlateRect& InClippingRect = DrawElement.GetClippingRect();
+
+	uint32 Layer = DrawElement.GetLayer();
+
+	// extract the layout transform from the draw element
+	FSlateLayoutTransform InverseLayoutTransform(Inverse(FSlateLayoutTransform(DrawElement.GetScale(), Position)));
+	FSlateRotatedClipRectType RenderClipRect = FSlateRotatedClipRectType::MakeSnappedRotatedRect(InClippingRect, InverseLayoutTransform, RenderTransform);
+
+	// Determine the four corners of the quad
+	FVector2D TopLeft = FVector2D::ZeroVector;
+	FVector2D TopRight = FVector2D(LocalSize.X, 0);
+	FVector2D BotLeft = FVector2D(0, LocalSize.Y);
+	FVector2D BotRight = FVector2D(LocalSize.X, LocalSize.Y);
+
+
+	// Offset by half a texel if the platform requires it for pixel perfect sampling
+	//FVector2D HalfTexel = FVector2D(PixelCenterOffset / WindowSize.X, PixelCenterOffset / WindowSize.Y);
+
+	FVector2D WorldTopLeft = ElementBatcherUtils::RoundToInt(TransformPoint(RenderTransform, TopLeft));
+	FVector2D WorldBotRight = ElementBatcherUtils::RoundToInt(TransformPoint(RenderTransform, BotRight));
+
+	// Clip post processing rect to the clip rect
+	WorldTopLeft = FMath::Max(WorldTopLeft, RenderClipRect.TopLeft);
+	WorldBotRight = FMath::Min(WorldBotRight, RenderClipRect.TopLeft + RenderClipRect.ExtentX + RenderClipRect.ExtentY);
+
+	FVector2D SizeUV = (WorldBotRight - WorldTopLeft) / WindowSize;
+
+	FShaderParams Params = FShaderParams::MakePixelShaderParams(FVector4(WorldTopLeft, WorldBotRight), FVector4(Payload.PostProcessData.X, Payload.PostProcessData.Y, Payload.DownsampleAmount, 0));
+
+	FElementBatchMap& LayerToElementBatches = DrawLayer->GetElementBatchMap();
+
+	// See if the layer already exists.
+	TUniqueObj<FElementBatchArray>* ElementBatches = LayerToElementBatches.Find(Layer);
+	if (!ElementBatches)
+	{
+		// The layer doesn't exist so make it now
+		ElementBatches = &LayerToElementBatches.Add(Layer);
+	}
+	check(ElementBatches);
+
+	// Custom elements are not batched together 
+	(*ElementBatches)->Add(FSlateElementBatch(nullptr, Params, ESlateShader::PostProcess, ESlateDrawPrimitive::TriangleList, ESlateDrawEffect::None, ESlateBatchDrawFlag::None, DrawElement.GetScissorRect()));
+	
+}
+
+FSlateElementBatch& FSlateElementBatcher::FindBatchForElement(
 	uint32 Layer, 
 	const FShaderParams& ShaderParams, 
 	const FSlateShaderResource* InTexture, 
@@ -1918,5 +1987,6 @@ FSlateElementBatch& FSlateElementBatcher::FindBatchForElement(
 void FSlateElementBatcher::ResetBatches()
 {
 	bRequiresVsync = false;
+	NumPostProcessPasses = 0;
 }
 
diff --git a/Engine/Source/Runtime/SlateCore/Private/Widgets/SWindow.cpp b/Engine/Source/Runtime/SlateCore/Private/Widgets/SWindow.cpp
--- a/Engine/Source/Runtime/SlateCore/Private/Widgets/SWindow.cpp
+++ b/Engine/Source/Runtime/SlateCore/Private/Widgets/SWindow.cpp
@@ -1661,7 +1661,6 @@ EWindowZone::Type SWindow::GetCurrentWindowZone(FVector2D LocalMousePosition)
 	const bool bIsBorderlessGameWindow = Type == EWindowType::GameWindow && !bHasOSWindowBorder;
 
 	const float WindowDPIScale = FSlateApplicationBase::Get().GetApplicationScale() * (NativeWindow.IsValid() ? NativeWindow->GetDPIScaleFactor() : 1.0f);
-
 	const FMargin DPIScaledResizeBorder = UserResizeBorder * WindowDPIScale;
 
 	const bool bIsCursorVisible = FSlateApplicationBase::Get().GetPlatformCursor()->GetType() != EMouseCursor::None;
@@ -1788,6 +1787,7 @@ SWindow::SWindow()
 	, ExpectedMaxWidth( INDEX_NONE )
 	, ExpectedMaxHeight( INDEX_NONE )
 	, TitleBar()
+	, bIsDrawingEnabled( true )
 {
 }
 
@@ -1860,6 +1860,8 @@ void SWindow::SetWindowMode( EWindowMode::Type NewWindowMode )
 			PreFullscreenPosition = ScreenPosition;
 		}
 
+		bIsDrawingEnabled = false;
+
 		NativeWindow->SetWindowMode( NewWindowMode );
 	
 		const FVector2D vp = IsMirrorWindow() ? GetSizeInScreen() : GetViewportSize();
@@ -1876,6 +1878,8 @@ void SWindow::SetWindowMode( EWindowMode::Type NewWindowMode )
 			// If we left fullscreen, reset the screen position;
 			MoveWindowTo(PreFullscreenPosition);
 		}
+
+		bIsDrawingEnabled = true;
 	}
 
 }
diff --git a/Engine/Source/Runtime/SlateCore/Public/Rendering/DrawElements.h b/Engine/Source/Runtime/SlateCore/Public/Rendering/DrawElements.h
--- a/Engine/Source/Runtime/SlateCore/Public/Rendering/DrawElements.h
+++ b/Engine/Source/Runtime/SlateCore/Public/Rendering/DrawElements.h
@@ -45,6 +45,7 @@ public:
 	// Brush data
 	const FSlateBrush* BrushResource;
 	const FSlateShaderResourceProxy* ResourceProxy;
+	FSlateShaderResource* RenderTargetResource;
 
 	// Box Data
 	FVector2D RotationPoint;
@@ -68,10 +69,12 @@ public:
 	TArray<FVector2D> Points;
 
 	// Viewport data
-	FSlateShaderResource* ViewportRenderTargetTexture;
-	bool bAllowViewportScaling;
-	bool bViewportTextureAlphaOnly;
-	bool bRequiresVSync;
+	bool bAllowViewportScaling:1;
+	bool bViewportTextureAlphaOnly:1;
+	bool bRequiresVSync:1;
+	
+	// Whether or not to anti-alias lines
+	bool bAntialias:1;
 
 	// Misc data
 	ESlateBatchDrawFlag::Type BatchFlags;
@@ -95,9 +98,13 @@ public:
 	// Layer handle
 	FSlateDrawLayerHandle* LayerHandle;
 
+	// Post Process Data
+	FVector4 PostProcessData;
+	int32 DownsampleAmount;
+
 	// Line data
 	ESlateLineJoinType::Type SegmentJoinType;
-	bool bAntialias;
+
 
 	SLATECORE_API static FSlateShaderResourceManager* ResourceManager;
 
@@ -105,9 +112,9 @@ public:
 		: Tint(FLinearColor::White)
 		, BrushResource(nullptr)
 		, ResourceProxy(nullptr)
+		, RenderTargetResource(nullptr)
 		, RotationPoint(FVector2D::ZeroVector)
 		, ImmutableText(nullptr)
-		, ViewportRenderTargetTexture(nullptr)
 		, bViewportTextureAlphaOnly(false)
 		, bRequiresVSync(false)
 		, BatchFlags(ESlateBatchDrawFlag::None)
@@ -115,10 +122,6 @@ public:
 		, InstanceData(nullptr)
 		, InstanceOffset(0)
 		, NumInstances(0)
-		//, CachedRenderDataOffset(FVector2D::ZeroVector)
-		//, LayerHandle(nullptr)
-		//, SegmentJoinType(ESlateLineJoinType::Sharp)
-		//, bAntialias(true)
 	{ }
 
 	void SetBoxPayloadProperties( const FSlateBrush* InBrush, const FLinearColor& InTint, FSlateShaderResourceProxy* InResourceProxy = nullptr, ISlateUpdatableInstanceBuffer* InInstanceData = nullptr )
@@ -195,7 +198,7 @@ public:
 	void SetViewportPayloadProperties( const TSharedPtr<const ISlateViewport>& InViewport, const FLinearColor& InTint )
 	{
 		Tint = InTint;
-		ViewportRenderTargetTexture = InViewport->GetViewportRenderTargetTexture();
+		RenderTargetResource = InViewport->GetViewportRenderTargetTexture();
 		bAllowViewportScaling = InViewport->AllowScaling();
 		bViewportTextureAlphaOnly = InViewport->IsViewportTextureAlphaOnly();
 		bRequiresVSync = InViewport->RequiresVsync();
@@ -252,6 +255,7 @@ public:
 		ET_CustomVerts,
 		ET_CachedBuffer,
 		ET_Layer,
+		ET_PostProcessPass,
 		ET_Count,
 	};
 
@@ -439,6 +443,7 @@ public:
 
 	SLATECORE_API static void MakeLayer(FSlateWindowElementList& ElementList, uint32 InLayer, TSharedPtr<FSlateDrawLayerHandle, ESPMode::ThreadSafe>& DrawLayerHandle);
 
+	SLATECORE_API static void MakePostProcessPass(FSlateWindowElementList& ElementList, uint32 InLayer, const FPaintGeometry& PaintGeometry, const FSlateRect& InClippingRect, const FVector4& Params, int32 DownsampleAmount);
 
 	FORCEINLINE EElementType GetElementType() const { return ElementType; }
 	FORCEINLINE uint32 GetLayer() const { return Layer; }
@@ -480,23 +485,26 @@ struct FShaderParams
 {
 	/** Pixel shader parameters */
 	FVector4 PixelParams;
+	FVector4 PixelParams2;
 
 	FShaderParams()
 		: PixelParams( 0,0,0,0 )
+		, PixelParams2( 0,0,0,0 ) 
 	{}
 
-	FShaderParams( const FVector4& InPixelParams )
+	FShaderParams( const FVector4& InPixelParams, const FVector4& InPixelParams2 = FVector4(0) )
 		: PixelParams( InPixelParams )
+		, PixelParams2( InPixelParams2 )
 	{}
 
 	bool operator==( const FShaderParams& Other ) const
 	{
-		return PixelParams == Other.PixelParams;
+		return PixelParams == Other.PixelParams && PixelParams2 == Other.PixelParams2;
 	}
 
-	static FShaderParams MakePixelShaderParams( const FVector4& PixelShaderParams )
+	static FShaderParams MakePixelShaderParams( const FVector4& PixelShaderParams, const FVector4& InPixelShaderParams2 = FVector4(0) )
 	{
-		return FShaderParams( PixelShaderParams );
+		return FShaderParams( PixelShaderParams, InPixelShaderParams2);
 	}
 };
 
diff --git a/Engine/Source/Runtime/SlateCore/Public/Rendering/ElementBatcher.h b/Engine/Source/Runtime/SlateCore/Public/Rendering/ElementBatcher.h
--- a/Engine/Source/Runtime/SlateCore/Public/Rendering/ElementBatcher.h
+++ b/Engine/Source/Runtime/SlateCore/Public/Rendering/ElementBatcher.h
@@ -27,13 +27,16 @@ public:
 	 */
 	bool RequiresVsync() const { return bRequiresVsync; }
 
+	/** Whether or not any post process passes were batched */
+	bool HasFXPassses() const { return NumPostProcessPasses > 0;}
+
 	/** 
 	 * Resets all stored data accumulated during the batching process
 	 */
 	void ResetBatches();
 
 private:
-	void AddElements(FSlateDrawLayer& InDrawLayer);
+	void AddElements(const FSlateWindowElementList& ElementList, FSlateDrawLayer& InDrawLayer);
 	
 	FColor PackVertexColor(const FLinearColor& InLinearColor);
 
@@ -100,6 +103,8 @@ private:
 
 	void AddLayer(const FSlateDrawElement& DrawElement);
 
+	void AddPostProcessPass(const FSlateDrawElement& DrawElement, const FVector2D& WindowSize);
+
 	/** 
 	 * Finds an batch for an element based on the passed in parameters
 	 * Elements with common parameters and layers will batched together.
@@ -138,6 +143,9 @@ private:
 	/** Track the number of drawn texts from the previous frame to report to stats. */
 	int32 NumDrawnTextsStat;
 
+	/** How many post process passes are needed */
+	int32 NumPostProcessPasses;
+
 	/** Offset to use when supporting 1:1 texture to pixel snapping */
 	const float PixelCenterOffset;
 
diff --git a/Engine/Source/Runtime/SlateCore/Public/Rendering/RenderingCommon.h b/Engine/Source/Runtime/SlateCore/Public/Rendering/RenderingCommon.h
--- a/Engine/Source/Runtime/SlateCore/Public/Rendering/RenderingCommon.h
+++ b/Engine/Source/Runtime/SlateCore/Public/Rendering/RenderingCommon.h
@@ -44,6 +44,8 @@ namespace ESlateShader
 	const Type LineSegment = 3;
 	/** For completely customized materials.  Makes no assumptions on use*/
 	const Type Custom = 4;
+	/** For post processing passes */
+	const Type PostProcess = 5;
 };
 
 /**
diff --git a/Engine/Source/Runtime/SlateCore/Public/Textures/SlateShaderResource.h b/Engine/Source/Runtime/SlateCore/Public/Textures/SlateShaderResource.h
--- a/Engine/Source/Runtime/SlateCore/Public/Textures/SlateShaderResource.h
+++ b/Engine/Source/Runtime/SlateCore/Public/Textures/SlateShaderResource.h
@@ -22,6 +22,9 @@ namespace ESlateShaderResource
 		/** Material resource. */
 		Material,
 
+		/** Post Process. */
+		PostProcess, 
+
 		/** No Resource. */
 		Invalid,
 	};
diff --git a/Engine/Source/Runtime/SlateCore/Public/Widgets/SWindow.h b/Engine/Source/Runtime/SlateCore/Public/Widgets/SWindow.h
--- a/Engine/Source/Runtime/SlateCore/Public/Widgets/SWindow.h
+++ b/Engine/Source/Runtime/SlateCore/Public/Widgets/SWindow.h
@@ -700,6 +700,8 @@ public:
 
 	virtual bool SupportsKeyboardFocus() const override;
 
+	bool IsDrawingEnabled() const { return bIsDrawingEnabled; }
+
 private:
 	virtual FReply OnFocusReceived( const FGeometry& MyGeometry, const FFocusEvent& InFocusEvent ) override;
 	virtual FReply OnMouseButtonDown( const FGeometry& MyGeometry, const FPointerEvent& MouseEvent ) override;
@@ -1064,6 +1066,9 @@ protected:
 	// The margin around the edges of the window that will be detected as places the user can grab to resize the window. 
 	FMargin UserResizeBorder;
 
+	// Whether or not drawing is enabled for this window
+	bool bIsDrawingEnabled;
+
 protected:
 	
 	void ConstructWindowInternals();
diff --git a/Engine/Source/Runtime/SlateRHIRenderer/Private/SlatePostProcessResource.cpp b/Engine/Source/Runtime/SlateRHIRenderer/Private/SlatePostProcessResource.cpp
new file mode 100644
--- /dev/null
+++ b/Engine/Source/Runtime/SlateRHIRenderer/Private/SlatePostProcessResource.cpp
@@ -0,0 +1,82 @@
+// Copyright 1998-2016 Epic Games, Inc. All Rights Reserved.
+
+#include "SlateRHIRendererPrivatePCH.h"
+#include "SlatePostProcessResource.h"
+
+DECLARE_MEMORY_STAT(TEXT("PostProcess RenderTargets"), STAT_SLATEPPRenderTargetMem, STATGROUP_SlateMemory);
+
+FSlatePostProcessResource::FSlatePostProcessResource(int32 InRenderTargetCount)
+	: RenderTargetSize(FIntPoint::ZeroValue)
+	, RenderTargetCount(InRenderTargetCount)
+{
+	
+}
+
+FSlatePostProcessResource::~FSlatePostProcessResource()
+{
+
+}
+
+void FSlatePostProcessResource::Resize(const FIntPoint& NewSize)
+{
+	check(IsInRenderingThread());
+	if(RenderTargetSize != NewSize)
+	{
+		RenderTargetSize = NewSize;
+		UpdateRHI();
+	}
+}
+
+void FSlatePostProcessResource::CleanUp()
+{
+	BeginReleaseResource(this);
+
+	BeginCleanup(this);
+}
+
+void FSlatePostProcessResource::InitDynamicRHI()
+{
+	RenderTargets.Empty();
+
+	PixelFormat = PF_B8G8R8A8;
+	if(RenderTargetSize.X > 0 && RenderTargetSize.Y > 0)
+	{
+		for (int32 TexIndex = 0; TexIndex < RenderTargetCount; ++TexIndex)
+		{
+			FTexture2DRHIRef RenderTargetTextureRHI;
+			FTexture2DRHIRef ShaderResourceUnused;
+			FRHIResourceCreateInfo CreateInfo;
+			RHICreateTargetableShaderResource2D(
+				RenderTargetSize.X,
+				RenderTargetSize.Y,
+				PixelFormat,
+				1,
+				/*TexCreateFlags=*/0,
+				TexCreate_RenderTargetable,
+				/*bNeedsTwoCopies=*/false,
+				CreateInfo,
+				RenderTargetTextureRHI,
+				ShaderResourceUnused
+			);
+
+			RenderTargets.Add(RenderTargetTextureRHI);
+		}
+	}
+
+	STAT(int64 TotalMemory = RenderTargetCount * GPixelFormats[PixelFormat].BlockBytes*RenderTargetSize.X*RenderTargetSize.Y);
+	SET_MEMORY_STAT(STAT_SLATEPPRenderTargetMem, TotalMemory);
+}
+
+void FSlatePostProcessResource::ReleaseDynamicRHI()
+{
+	SET_MEMORY_STAT(STAT_SLATEPPRenderTargetMem, 0);
+
+	RenderTargets.Empty();
+}
+
+void FSlatePostProcessResource::FinishCleanup()
+{
+	delete this;
+}
+
+
diff --git a/Engine/Source/Runtime/SlateRHIRenderer/Private/SlatePostProcessor.cpp b/Engine/Source/Runtime/SlateRHIRenderer/Private/SlatePostProcessor.cpp
new file mode 100644
--- /dev/null
+++ b/Engine/Source/Runtime/SlateRHIRenderer/Private/SlatePostProcessor.cpp
@@ -0,0 +1,390 @@
+// Copyright 1998-2016 Epic Games, Inc. All Rights Reserved.
+
+#include "SlateRHIRendererPrivatePCH.h"
+#include "SlatePostProcessor.h"
+#include "SlatePostProcessResource.h"
+
+DECLARE_CYCLE_STAT(TEXT("Slate PostProcessing RT"), STAT_SlatePostProcessingRTTime, STATGROUP_Slate);
+
+FSlatePostProcessor::FSlatePostProcessor()
+{
+	const int32 NumIntermediateTargets = 2;
+	IntermediateTargets = new FSlatePostProcessResource(NumIntermediateTargets);
+	BeginInitResource(IntermediateTargets);
+}
+
+FSlatePostProcessor::~FSlatePostProcessor()
+{
+	// Note this is deleted automatically because it implements FDeferredCleanupInterface.
+	IntermediateTargets->CleanUp();
+
+}
+
+void FSlatePostProcessor::BlurRect(FRHICommandListImmediate& RHICmdList, IRendererModule& RendererModule, const FBlurRectParams& Params, const FPostProcessRectParams& RectParams)
+{
+	SCOPE_CYCLE_COUNTER(STAT_SlatePostProcessingRTTime);
+
+	TArray<FVector4> WeightsAndOffsets;
+	const int32 SampleCount = ComputeBlurWeights(Params.KernelSize, Params.Strength, WeightsAndOffsets);
+
+	SCOPED_DRAW_EVENTF(RHICmdList, SlatePostProcess, TEXT("Slate Post Process Blur Background %dx%d"), SampleCount, SampleCount);
+
+	const bool bDownsample = Params.DownsampleAmount > 0;
+
+	FIntPoint DestRectSize = RectParams.DestRect.GetSize().IntPoint();
+	const FIntPoint CurrentSize = FIntPoint(IntermediateTargets->GetWidth(), IntermediateTargets->GetHeight());
+	FIntPoint RequiredSize = bDownsample
+										? FIntPoint(FMath::DivideAndRoundUp(DestRectSize.X, Params.DownsampleAmount), FMath::DivideAndRoundUp(DestRectSize.Y, Params.DownsampleAmount))
+										: DestRectSize;
+
+	// The max size can get ridiculous with large scale values.  Clamp to size of the backbuffer
+	RequiredSize.X = FMath::Min(RequiredSize.X, RectParams.SourceTextureSize.X);
+	RequiredSize.Y = FMath::Min(RequiredSize.Y, RectParams.SourceTextureSize.Y);
+	
+	const FIntPoint DownsampleSize = RequiredSize;
+	if(RequiredSize.X > CurrentSize.X || RequiredSize.Y > CurrentSize.Y)
+	{		
+		IntermediateTargets->Resize(RequiredSize);
+	}
+
+	if(bDownsample)
+	{
+		DownsampleRect(RHICmdList, RendererModule, RectParams, DownsampleSize);
+	}
+
+#if 1
+	FSamplerStateRHIRef BilinearClamp = TStaticSamplerState<SF_Bilinear, AM_Clamp, AM_Clamp, AM_Clamp>::GetRHI();
+
+	TShaderMap<FGlobalShaderType>* ShaderMap = GetGlobalShaderMap(GMaxRHIFeatureLevel);
+	check(ShaderMap);
+
+	TShaderMapRef<FScreenVS> VertexShader(ShaderMap);
+	TShaderMapRef<FSlatePostProcessBlurPS> PixelShader(ShaderMap);
+
+	const int32 SrcTextureWidth = RectParams.SourceTextureSize.X;
+	const int32 SrcTextureHeight = RectParams.SourceTextureSize.Y;
+
+	const int32 DestTextureWidth = IntermediateTargets->GetWidth();
+	const int32 DestTextureHeight = IntermediateTargets->GetHeight();
+
+	const FSlateRect& SourceRect = RectParams.SourceRect;
+	const FSlateRect& DestRect = RectParams.DestRect;
+
+	FVertexDeclarationRHIRef VertexDecl = RendererModule.GetFilterVertexDeclaration().VertexDeclarationRHI;
+	check(IsValidRef(VertexDecl));
+	
+	RHICmdList.SetBlendState(TStaticBlendState<>::GetRHI());
+	RHICmdList.SetRasterizerState(TStaticRasterizerState<>::GetRHI());
+	RHICmdList.SetDepthStencilState(TStaticDepthStencilState<false, CF_Always>::GetRHI());
+	RHICmdList.SetScissorRect(false, 0, 0, 0, 0);
+	RHICmdList.SetViewport(0, 0, 0, DestTextureWidth, DestTextureHeight, 0.0f);
+	
+	const FVector2D InvBufferSize = FVector2D(1.0f / DestTextureWidth, 1.0f / DestTextureHeight);
+	const FVector2D HalfTexelOffset = FVector2D(0.5f/ DestTextureWidth, 0.5f/ DestTextureHeight);
+
+	for (int32 PassIndex = 0; PassIndex < 2; ++PassIndex)
+	{
+		// First pass render to the render target with the post process fx
+		if (PassIndex == 0)
+		{
+			FTexture2DRHIRef SourceTexture = bDownsample ? IntermediateTargets->GetRenderTarget(0) : RectParams.SourceTexture;
+			FTexture2DRHIRef DestTexture = IntermediateTargets->GetRenderTarget(1);
+
+			RHICmdList.TransitionResource(EResourceTransitionAccess::EReadable, SourceTexture);
+			RHICmdList.TransitionResource(EResourceTransitionAccess::EWritable, DestTexture);
+
+
+			SetRenderTarget(RHICmdList, DestTexture, FTextureRHIRef());
+		
+			// This is necessary on metal because shader state is cleared on metal for some reason when render targets are switched
+			static FGlobalBoundShaderState BoundShaderState;
+			SetGlobalBoundShaderState(RHICmdList, GMaxRHIFeatureLevel, BoundShaderState, VertexDecl, *VertexShader, *PixelShader);
+
+			PixelShader->SetWeightsAndOffsets(RHICmdList, WeightsAndOffsets, SampleCount);
+			PixelShader->SetTexture(RHICmdList, SourceTexture, BilinearClamp);
+		
+			if(bDownsample)
+			{
+				PixelShader->SetUVBounds(RHICmdList, FVector4(FVector2D::ZeroVector, FVector2D((float)DownsampleSize.X/DestTextureWidth, (float)DownsampleSize.Y / DestTextureHeight) - HalfTexelOffset));
+				PixelShader->SetBufferSizeAndDirection(RHICmdList, InvBufferSize, FVector2D(1, 0));
+
+				RendererModule.DrawRectangle(
+					RHICmdList,
+					0, 0,
+					DownsampleSize.X, DownsampleSize.Y,
+					0, 0,
+					DownsampleSize.X, DownsampleSize.Y,
+					FIntPoint(DestTextureWidth, DestTextureHeight),
+					FIntPoint(DestTextureWidth, DestTextureHeight),
+					*VertexShader,
+					EDRF_Default);
+			}
+			else
+			{
+				const FVector2D InvSrcTetureSize(1.f / SrcTextureWidth, 1.f / SrcTextureHeight);
+
+				const FVector2D UVStart = FVector2D(DestRect.Left, DestRect.Top) * InvSrcTetureSize;
+				const FVector2D UVEnd = FVector2D(DestRect.Right, DestRect.Bottom) * InvSrcTetureSize;
+				const FVector2D SizeUV = UVEnd - UVStart;
+
+				PixelShader->SetUVBounds(RHICmdList, FVector4(UVStart, UVEnd));
+				PixelShader->SetBufferSizeAndDirection(RHICmdList, InvSrcTetureSize, FVector2D(1, 0));
+
+				RendererModule.DrawRectangle(
+					RHICmdList,
+					0, 0,
+					RequiredSize.X, RequiredSize.Y,
+					UVStart.X, UVStart.Y,
+					SizeUV.X, SizeUV.Y,
+					FIntPoint(DestTextureWidth, DestTextureHeight),
+					FIntPoint(1, 1),
+					*VertexShader,
+					EDRF_Default);
+			}
+
+		}
+		else
+		{
+			FTexture2DRHIRef SourceTexture = IntermediateTargets->GetRenderTarget(1);
+			FTexture2DRHIRef DestTexture = IntermediateTargets->GetRenderTarget(0);
+
+			RHICmdList.TransitionResource(EResourceTransitionAccess::EReadable, SourceTexture);
+			RHICmdList.TransitionResource(EResourceTransitionAccess::EWritable, DestTexture);
+
+			SetRenderTarget(RHICmdList, DestTexture, FTextureRHIRef());
+
+			// This is necessary on metal because shader state is cleared on metal for some reason when render targets are switched
+			static FGlobalBoundShaderState BoundShaderState;
+			SetGlobalBoundShaderState(RHICmdList, GMaxRHIFeatureLevel, BoundShaderState, VertexDecl, *VertexShader, *PixelShader);
+			
+			PixelShader->SetWeightsAndOffsets(RHICmdList, WeightsAndOffsets, SampleCount);
+			PixelShader->SetUVBounds(RHICmdList, FVector4(FVector2D::ZeroVector, FVector2D((float)DownsampleSize.X/DestTextureWidth, (float)DownsampleSize.Y / DestTextureHeight) - HalfTexelOffset));
+			PixelShader->SetTexture(RHICmdList, SourceTexture, BilinearClamp);
+			PixelShader->SetBufferSizeAndDirection(RHICmdList, InvBufferSize, FVector2D(0, 1));
+
+			RendererModule.DrawRectangle(
+				RHICmdList,
+				0, 0,
+				DownsampleSize.X, DownsampleSize.Y,
+				0, 0,
+				DownsampleSize.X, DownsampleSize.Y,
+				FIntPoint(DestTextureWidth, DestTextureHeight),
+				FIntPoint(DestTextureWidth, DestTextureHeight),
+				*VertexShader,
+				EDRF_Default);
+		}
+
+	
+	}
+
+#endif
+	UpsampleRect(RHICmdList, RendererModule, RectParams, DownsampleSize);
+}
+
+void FSlatePostProcessor::DownsampleRect(FRHICommandListImmediate& RHICmdList, IRendererModule& RendererModule, const FPostProcessRectParams& Params, const FIntPoint& DownsampleSize)
+{
+	SCOPED_DRAW_EVENT(RHICmdList, SlatePostProcessDownsample);
+
+	// Source is the viewport.  This is the width and height of the viewport backbuffer
+	const int32 SrcTextureWidth = Params.SourceTextureSize.X;
+	const int32 SrcTextureHeight = Params.SourceTextureSize.Y;
+
+	// Dest is the destination quad for the downsample
+	const int32 DestTextureWidth = IntermediateTargets->GetWidth();
+	const int32 DestTextureHeight = IntermediateTargets->GetHeight();
+
+	// Rect of the viewport
+	const FSlateRect& SourceRect = Params.SourceRect;
+
+	// Rect of the final destination post process effect (not downsample rect).  This is the area we sample from
+	const FSlateRect& DestRect = Params.DestRect;
+
+	TShaderMap<FGlobalShaderType>* ShaderMap = GetGlobalShaderMap(GMaxRHIFeatureLevel);
+	TShaderMapRef<FScreenVS> VertexShader(ShaderMap);
+
+	FSamplerStateRHIRef BilinearClamp = TStaticSamplerState<SF_Bilinear, AM_Clamp, AM_Clamp, AM_Clamp>::GetRHI();
+
+	FTexture2DRHIRef DestTexture = IntermediateTargets->GetRenderTarget(0);
+
+	// Downsample and store in intermediate texture
+	{
+		TShaderMapRef<FSlatePostProcessDownsamplePS> PixelShader(ShaderMap);
+
+		RHICmdList.TransitionResource(EResourceTransitionAccess::EReadable, Params.SourceTexture);
+		RHICmdList.TransitionResource(EResourceTransitionAccess::EWritable, DestTexture);
+
+		const FVector2D InvSrcTetureSize(1.f/SrcTextureWidth, 1.f/SrcTextureHeight);
+
+		const FVector2D UVStart = FVector2D(DestRect.Left, DestRect.Top) * InvSrcTetureSize;
+		const FVector2D UVEnd = FVector2D(DestRect.Right, DestRect.Bottom) * InvSrcTetureSize;
+		const FVector2D SizeUV = UVEnd - UVStart;
+		
+		RHICmdList.SetViewport(0, 0, 0, DestTextureWidth, DestTextureHeight, 0.0f);
+		RHICmdList.SetScissorRect(false, 0, 0, 0, 0);
+
+		SetRenderTarget(RHICmdList, DestTexture, FTextureRHIRef());
+		
+		RHICmdList.SetBlendState(TStaticBlendState<>::GetRHI());
+		RHICmdList.SetRasterizerState(TStaticRasterizerState<>::GetRHI());
+		RHICmdList.SetDepthStencilState(TStaticDepthStencilState<false, CF_Always>::GetRHI());
+		
+		static FGlobalBoundShaderState BoundShaderState;
+		SetGlobalBoundShaderState(RHICmdList, GMaxRHIFeatureLevel, BoundShaderState, RendererModule.GetFilterVertexDeclaration().VertexDeclarationRHI, *VertexShader, *PixelShader);
+		
+		PixelShader->SetShaderParams(RHICmdList, FVector4(InvSrcTetureSize.X, InvSrcTetureSize.Y, 0, 0));
+		PixelShader->SetUVBounds(RHICmdList, FVector4(UVStart, UVEnd));
+		PixelShader->SetTexture(RHICmdList, Params.SourceTexture, BilinearClamp);
+
+		RendererModule.DrawRectangle(
+			RHICmdList,
+			0, 0,
+			DownsampleSize.X, DownsampleSize.Y,
+			UVStart.X, UVStart.Y,
+			SizeUV.X, SizeUV.Y,
+			FIntPoint(DestTextureWidth, DestTextureHeight),
+			FIntPoint(1,1), 
+			*VertexShader,
+			EDRF_Default);
+	}
+	
+	// Testing only
+#if 0
+	UpsampleRect(RHICmdList, RendererModule, Params, DownsampleSize);
+#endif
+}
+
+void FSlatePostProcessor::UpsampleRect(FRHICommandListImmediate& RHICmdList, IRendererModule& RendererModule, const FPostProcessRectParams& Params, const FIntPoint& DownsampleSize)
+{
+	SCOPED_DRAW_EVENT(RHICmdList, SlatePostProcessUpsample);
+
+	// Original source texture is now the destination texture
+	FTexture2DRHIRef DestTexture = Params.SourceTexture;
+	const int32 DestTextureWidth = Params.SourceTextureSize.X;
+	const int32 DestTextureHeight = Params.SourceTextureSize.Y;
+
+	const int32 DownsampledWidth = DownsampleSize.X;
+	const int32 DownsampledHeight = DownsampleSize.Y;
+
+	// Source texture is the texture that was originally downsampled
+	FTexture2DRHIRef SrcTexture = IntermediateTargets->GetRenderTarget(0);
+	const int32 SrcTextureWidth = IntermediateTargets->GetWidth();
+	const int32 SrcTextureHeight = IntermediateTargets->GetHeight();
+
+	const FSlateRect& SourceRect = Params.SourceRect;
+	const FSlateRect& DestRect = Params.DestRect;
+
+	TShaderMap<FGlobalShaderType>* ShaderMap = GetGlobalShaderMap(GMaxRHIFeatureLevel);
+	TShaderMapRef<FScreenVS> VertexShader(ShaderMap);
+
+	FSamplerStateRHIRef BilinearClamp = TStaticSamplerState<SF_Bilinear, AM_Clamp, AM_Clamp, AM_Clamp>::GetRHI();
+
+	RHICmdList.SetViewport(0, 0, 0, DestTextureWidth, DestTextureHeight, 0.0f);
+
+	RHICmdList.TransitionResource(EResourceTransitionAccess::EWritable, Params.SourceTexture);
+	RHICmdList.TransitionResource(EResourceTransitionAccess::EReadable, DestTexture);
+
+	SetRenderTarget(RHICmdList, DestTexture, FTextureRHIRef());
+
+	Params.RestoreStateFunc();
+
+	TShaderMapRef<FScreenPS> PixelShader(ShaderMap);
+
+	static FGlobalBoundShaderState BoundShaderState;
+	SetGlobalBoundShaderState(RHICmdList, GMaxRHIFeatureLevel, BoundShaderState, RendererModule.GetFilterVertexDeclaration().VertexDeclarationRHI, *VertexShader, *PixelShader);
+
+	PixelShader->SetParameters(RHICmdList, BilinearClamp, SrcTexture);
+
+	RendererModule.DrawRectangle(RHICmdList,
+		DestRect.Left, DestRect.Top,
+		DestRect.Right - DestRect.Left, DestRect.Bottom - DestRect.Top,
+		0, 0,
+		(float)DownsampledWidth/SrcTextureWidth-1.0/ SrcTextureWidth, (float)DownsampledHeight/SrcTextureHeight-1.0 / SrcTextureHeight,
+		Params.SourceTextureSize,
+		FIntPoint(1, 1),
+		*VertexShader,
+		EDRF_Default);
+
+}
+#define BILINEAR_FILTER_METHOD 1
+
+#if !BILINEAR_FILTER_METHOD
+
+static int32 ComputeWeights(int32 KernelSize, float Sigma, TArray<FVector4>& OutWeightsAndOffsets)
+{
+	OutWeightsAndOffsets.AddUninitialized(KernelSize / 2 + 1);
+
+	int32 SampleIndex = 0;
+	for (int32 X = 0; X < KernelSize; X += 2)
+	{
+		float Dist = X;
+		FVector4 WeightAndOffset;
+		WeightAndOffset.X = (1.0f / FMath::Sqrt(2 * PI*Sigma*Sigma))*FMath::Exp(-(Dist*Dist) / (2 * Sigma*Sigma));
+		WeightAndOffset.Y = Dist;
+
+		Dist = X + 1;
+		WeightAndOffset.Z = (1.0f / FMath::Sqrt(2 * PI*Sigma*Sigma))*FMath::Exp(-(Dist*Dist) / (2 * Sigma*Sigma));
+		WeightAndOffset.W = Dist;
+
+		OutWeightsAndOffsets[SampleIndex] = WeightAndOffset;
+
+		++SampleIndex;
+	}
+
+	return KernelSize;
+};
+
+#else
+
+static float GetWeight(float Dist, float Strength)
+{
+	// from https://en.wikipedia.org/wiki/Gaussian_blur
+	float Strength2 = Strength*Strength;
+	return (1.0f / FMath::Sqrt(2 * PI*Strength2))*FMath::Exp(-(Dist*Dist) / (2 * Strength2));
+}
+
+static FVector2D GetWeightAndOffset(float Dist, float Sigma)
+{
+	float Offset1 = Dist;
+	float Weight1 = GetWeight(Offset1, Sigma);
+
+	float Offset2 = Dist + 1;
+	float Weight2 = GetWeight(Offset2, Sigma);
+
+	float TotalWeight = Weight1 + Weight2;
+
+	float Offset = 0;
+	if (TotalWeight > 0)
+	{
+		Offset = (Weight1*Offset1 + Weight2*Offset2) / TotalWeight;
+	}
+
+
+	return FVector2D(TotalWeight, Offset);
+}
+
+static int32 ComputeWeights(int32 KernelSize, float Sigma, TArray<FVector4>& OutWeightsAndOffsets)
+{
+	int32 NumSamples = FMath::DivideAndRoundUp(KernelSize, 2);
+
+	// We need half of the sample count array because we're packing two samples into one float4
+
+	OutWeightsAndOffsets.AddUninitialized(NumSamples%2 == 0 ? NumSamples / 2 : NumSamples/2+1);
+
+	OutWeightsAndOffsets[0] = FVector4(FVector2D(GetWeight(0,Sigma), 0), GetWeightAndOffset(1, Sigma) );
+	int32 SampleIndex = 1;
+	for (int32 X = 3; X < KernelSize; X += 4)
+	{
+		OutWeightsAndOffsets[SampleIndex] = FVector4(GetWeightAndOffset(X, Sigma), GetWeightAndOffset(X + 2, Sigma));
+
+		++SampleIndex;
+	}
+
+	return NumSamples;
+};
+
+#endif
+
+int32 FSlatePostProcessor::ComputeBlurWeights(int32 KernelSize, float StdDev, TArray<FVector4>& OutWeightsAndOffsets)
+{
+	return ComputeWeights(KernelSize, StdDev, OutWeightsAndOffsets);
+}
diff --git a/Engine/Source/Runtime/SlateRHIRenderer/Private/SlatePostProcessor.h b/Engine/Source/Runtime/SlateRHIRenderer/Private/SlatePostProcessor.h
new file mode 100644
--- /dev/null
+++ b/Engine/Source/Runtime/SlateRHIRenderer/Private/SlatePostProcessor.h
@@ -0,0 +1,38 @@
+// Copyright 1998-2016 Epic Games, Inc. All Rights Reserved.
+
+#pragma once
+
+class FSlatePostProcessResource;
+class IRendererModule;
+
+struct FPostProcessRectParams
+{
+	FTexture2DRHIRef SourceTexture;
+	FSlateRect SourceRect;
+	FSlateRect DestRect;
+	FIntPoint SourceTextureSize;
+	TFunction<void()> RestoreStateFunc; 
+};
+
+struct FBlurRectParams
+{
+	int32 KernelSize;
+	int32 DownsampleAmount;
+	float Strength;
+};
+
+class FSlatePostProcessor
+{
+public:
+	FSlatePostProcessor();
+	~FSlatePostProcessor();
+
+	void BlurRect(FRHICommandListImmediate& RHICmdList, IRendererModule& RendererModule, const FBlurRectParams& Params, const FPostProcessRectParams& RectParams);
+
+private:
+	void DownsampleRect(FRHICommandListImmediate& RHICmdList, IRendererModule& RendererModule, const FPostProcessRectParams& Params, const FIntPoint& DownsampleSize);
+	void UpsampleRect(FRHICommandListImmediate& RHICmdList, IRendererModule& RendererModule, const FPostProcessRectParams& Params, const FIntPoint& DownsampleSize);
+	int32 ComputeBlurWeights(int32 KernelSize, float StdDev, TArray<FVector4>& OutWeightsAndOffsets);
+private:
+	FSlatePostProcessResource* IntermediateTargets;
+};
\ No newline at end of file
diff --git a/Engine/Source/Runtime/SlateRHIRenderer/Private/SlateRHIRenderer.cpp b/Engine/Source/Runtime/SlateRHIRenderer/Private/SlateRHIRenderer.cpp
--- a/Engine/Source/Runtime/SlateRHIRenderer/Private/SlateRHIRenderer.cpp
+++ b/Engine/Source/Runtime/SlateRHIRenderer/Private/SlateRHIRenderer.cpp
@@ -412,7 +412,7 @@ void FSlateRHIRenderer::UpdateFullscreenState( const TSharedRef<SWindow> Window,
 
 void FSlateRHIRenderer::RestoreSystemResolution(const TSharedRef<SWindow> InWindow)
 {
-	if (!GIsEditor)
+	if (!GIsEditor && InWindow->GetWindowMode() == EWindowMode::Fullscreen)
 	{
 		// Force the window system to resize the active viewport, even though nothing might have appeared to change.
 		// On windows, DXGI might change the window resolution behind our backs when we alt-tab out. This will make
diff --git a/Engine/Source/Runtime/SlateRHIRenderer/Private/SlateRHIRenderingPolicy.cpp b/Engine/Source/Runtime/SlateRHIRenderer/Private/SlateRHIRenderingPolicy.cpp
--- a/Engine/Source/Runtime/SlateRHIRenderer/Private/SlateRHIRenderingPolicy.cpp
+++ b/Engine/Source/Runtime/SlateRHIRenderer/Private/SlateRHIRenderingPolicy.cpp
@@ -14,6 +14,7 @@
 #include "SlateUpdatableBuffer.h"
 #include "SlateElementIndexBuffer.h"
 #include "SlateElementVertexBuffer.h"
+#include "SlatePostProcessor.h"
 
 extern void UpdateNoiseTextureParameters(FViewUniformShaderParameters& ViewUniformShaderParameters);
 
@@ -25,6 +26,7 @@ DECLARE_DWORD_COUNTER_STAT(TEXT("Num Vertices"), STAT_SlateVertexCount, STATGROU
 
 FSlateRHIRenderingPolicy::FSlateRHIRenderingPolicy(TSharedRef<FSlateFontServices> InSlateFontServices, TSharedRef<FSlateRHIResourceManager> InResourceManager, TOptional<int32> InitialBufferSize)
 	: FSlateRenderingPolicy(InSlateFontServices, 0)
+	, PostProcessor(new FSlatePostProcessor)
 	, ResourceManager(InResourceManager)
 	, bGammaCorrect(true)
 	, InitialBufferSizeOverride(InitialBufferSize)
@@ -217,6 +219,9 @@ void FSlateRHIRenderingPolicy::DrawElements(FRHICommandListImmediate& RHICmdList
 	// Should only be called by the rendering thread
 	check(IsInRenderingThread());
 
+	static const FName RendererModuleName("Renderer");
+	IRendererModule& RendererModule = FModuleManager::GetModuleChecked<IRendererModule>(RendererModuleName);
+
 	float TimeSeconds = FApp::GetCurrentTime() - GStartTime;
 	float DeltaTimeSeconds = FApp::GetDeltaTime();
 	float RealTimeSeconds =  FPlatformTime::Seconds() - GStartTime;
@@ -286,24 +291,32 @@ void FSlateRHIRenderingPolicy::DrawElements(FRHICommandListImmediate& RHICmdList
 		const ESlateShader::Type ShaderType = RenderBatch.ShaderType;
 		const FShaderParams& ShaderParams = RenderBatch.ShaderParams;
 
-		if( !RenderBatch.CustomDrawer.IsValid() )
+		auto UpdateScissorRect = [&RHICmdList, &RenderBatch]()
 		{
-			check(RenderBatch.NumIndices > 0);
-
-			FMatrix DynamicOffset = FTranslationMatrix::Make(FVector(RenderBatch.DynamicOffset.X, RenderBatch.DynamicOffset.Y, 0));
-
 			if (RenderBatch.ScissorRect.IsSet())
 			{
-				RHICmdList.SetScissorRect(true, RenderBatch.ScissorRect.GetValue().Left, RenderBatch.ScissorRect.GetValue().Top, RenderBatch.ScissorRect.GetValue().Right, RenderBatch.ScissorRect.GetValue().Bottom);
+				const FShortRect& ScissorRect = RenderBatch.ScissorRect.GetValue();
+				RHICmdList.SetScissorRect(true, ScissorRect.Left, ScissorRect.Top, ScissorRect.Right, ScissorRect.Bottom);
 			}
 			else
 			{
 				RHICmdList.SetScissorRect(false, 0, 0, 0, 0);
 			}
+		};
+
+		if( !RenderBatch.CustomDrawer.IsValid() )
+		{
+			FMatrix DynamicOffset = FTranslationMatrix::Make(FVector(RenderBatch.DynamicOffset.X, RenderBatch.DynamicOffset.Y, 0));
+			const FMatrix ViewProjection = DynamicOffset * ViewProjectionMatrix;
+
+			UpdateScissorRect();
+
+			const uint32 PrimitiveCount = RenderBatch.DrawPrimitiveType == ESlateDrawPrimitive::LineList ? RenderBatch.NumIndices / 2 : RenderBatch.NumIndices / 3; 
 
 			ESlateShaderResource::Type ResourceType = ShaderResource ? ShaderResource->GetType() : ESlateShaderResource::Invalid;
-			if( ResourceType != ESlateShaderResource::Material ) 
+			if( ResourceType != ESlateShaderResource::Material && ShaderType != ESlateShader::PostProcess ) 
 			{
+				check(RenderBatch.NumIndices > 0);
 				FSlateElementPS* PixelShader = GetTexturePixelShader(ShaderType, DrawEffects);
 
 				const bool bUseInstancing = RenderBatch.InstanceCount > 1 && RenderBatch.InstanceData != nullptr;
@@ -317,8 +330,7 @@ void FSlateRHIRenderingPolicy::DrawElements(FRHICommandListImmediate& RHICmdList
 					PixelShader->GetPixelShader(),
 					FGeometryShaderRHIRef()));
 
-
-				GlobalVertexShader->SetViewProjection(RHICmdList, DynamicOffset * ViewProjectionMatrix);
+				GlobalVertexShader->SetViewProjection(RHICmdList, ViewProjection);
 				GlobalVertexShader->SetVerticalAxisMultiplier(RHICmdList, bAllowSwitchVerticalAxis && RHINeedsToSwitchVerticalAxis(GShaderPlatformForFeatureLevel[GMaxRHIFeatureLevel]) ? -1.0f : 1.0f );
 
 #if !DEBUG_OVERDRAW
@@ -429,8 +441,7 @@ void FSlateRHIRenderingPolicy::DrawElements(FRHICommandListImmediate& RHICmdList
 				PixelShader->SetShaderParams(RHICmdList, ShaderParams.PixelParams);
 				PixelShader->SetDisplayGamma(RHICmdList, (DrawFlags & ESlateBatchDrawFlag::NoGamma) ? 1.0f : DisplayGamma);
 
-				uint32 PrimitiveCount = RenderBatch.DrawPrimitiveType == ESlateDrawPrimitive::LineList ? RenderBatch.NumIndices / 2 : RenderBatch.NumIndices / 3;
-
+	
 				// for RHIs that can't handle VertexOffset, we need to offset the stream source each time
 				if (!GRHISupportsBaseVertexIndex)
 				{
@@ -443,8 +454,9 @@ void FSlateRHIRenderingPolicy::DrawElements(FRHICommandListImmediate& RHICmdList
 					RHICmdList.DrawIndexedPrimitive(IndexBuffer->IndexBufferRHI, GetRHIPrimitiveType(RenderBatch.DrawPrimitiveType), RenderBatch.VertexOffset, 0, RenderBatch.NumVertices, RenderBatch.IndexOffset, PrimitiveCount, RenderBatch.InstanceCount);
 				}
 			}
-			else if (ShaderResource && ShaderResource->GetType() == ESlateShaderResource::Material && GEngine)
+			else if (GEngine && ShaderResource && ShaderResource->GetType() == ESlateShaderResource::Material && ShaderType != ESlateShader::PostProcess)
 			{
+				check(RenderBatch.NumIndices > 0);
 				// Note: This code is only executed if the engine is loaded (in early loading screens attempting to use a material is unsupported
 				if (!SceneView)
 				{
@@ -480,7 +492,7 @@ void FSlateRHIRenderingPolicy::DrawElements(FRHICommandListImmediate& RHICmdList
 							PixelShader->GetPixelShader(),
 							FGeometryShaderRHIRef()));
 
-						VertexShader->SetViewProjection(RHICmdList, DynamicOffset * ViewProjectionMatrix);
+						VertexShader->SetViewProjection(RHICmdList, ViewProjection);
 						VertexShader->SetVerticalAxisMultiplier(RHICmdList, bAllowSwitchVerticalAxis && RHINeedsToSwitchVerticalAxis(GShaderPlatformForFeatureLevel[GMaxRHIFeatureLevel]) ? -1.0f : 1.0f);
 						VertexShader->SetMaterialShaderParameters(RHICmdList, *SceneView, MaterialRenderProxy, Material);
 
@@ -498,8 +510,6 @@ void FSlateRHIRenderingPolicy::DrawElements(FRHICommandListImmediate& RHICmdList
 							PixelShader->SetAdditionalTexture(RHICmdList, TextureRHI, BilinearClamp);
 						}
 
-						uint32 PrimitiveCount = RenderBatch.DrawPrimitiveType == ESlateDrawPrimitive::LineList ? RenderBatch.NumIndices / 2 : RenderBatch.NumIndices / 3;
-
 						if ( bUseInstancing )
 						{
 							uint32 InstanceCount = RenderBatch.InstanceCount;
@@ -560,7 +570,25 @@ void FSlateRHIRenderingPolicy::DrawElements(FRHICommandListImmediate& RHICmdList
 					}
 				}
 			}
+			else if(ShaderType == ESlateShader::PostProcess)
+			{
+				const FVector4& QuadPositionData = ShaderParams.PixelParams;
+
+				FPostProcessRectParams RectParams;
+				RectParams.SourceTexture = BackBuffer.GetRenderTargetTexture();
+				RectParams.SourceRect = FSlateRect(0,0,BackBuffer.GetSizeXY().X, BackBuffer.GetSizeXY().Y);
+				RectParams.DestRect = FSlateRect(QuadPositionData.X, QuadPositionData.Y, QuadPositionData.Z, QuadPositionData.W);
+				RectParams.SourceTextureSize = BackBuffer.GetSizeXY();
+
+				RectParams.RestoreStateFunc = UpdateScissorRect;
 
+				FBlurRectParams BlurParams;
+				BlurParams.KernelSize = ShaderParams.PixelParams2.X;
+				BlurParams.Strength = ShaderParams.PixelParams2.Y;
+				BlurParams.DownsampleAmount = ShaderParams.PixelParams2.Z;
+
+				PostProcessor->BlurRect(RHICmdList, RendererModule, BlurParams, RectParams);
+			}
 		}
 		else
 		{
diff --git a/Engine/Source/Runtime/SlateRHIRenderer/Private/SlateRHIRenderingPolicy.h b/Engine/Source/Runtime/SlateRHIRenderer/Private/SlateRHIRenderingPolicy.h
--- a/Engine/Source/Runtime/SlateRHIRenderer/Private/SlateRHIRenderingPolicy.h
+++ b/Engine/Source/Runtime/SlateRHIRenderer/Private/SlateRHIRenderingPolicy.h
@@ -6,6 +6,7 @@
 #include "RenderingPolicy.h"
 
 class FSlateRHIResourceManager;
+class FSlatePostProcessor;
 
 class FSlateRHIRenderingPolicy : public FSlateRenderingPolicy
 {
@@ -53,6 +54,9 @@ private:
 	TSlateElementVertexBuffer<FSlateVertex> VertexBuffers;
 	FSlateElementIndexBuffer IndexBuffers;
 
+	/** Handles post process effects for slate */
+	TSharedRef<FSlatePostProcessor> PostProcessor;
+
 	TSharedRef<FSlateRHIResourceManager> ResourceManager;
 
 	bool bGammaCorrect;
diff --git a/Engine/Source/Runtime/SlateRHIRenderer/Private/SlateShaders.cpp b/Engine/Source/Runtime/SlateRHIRenderer/Private/SlateShaders.cpp
--- a/Engine/Source/Runtime/SlateRHIRenderer/Private/SlateShaders.cpp
+++ b/Engine/Source/Runtime/SlateRHIRenderer/Private/SlateShaders.cpp
@@ -11,6 +11,11 @@ IMPLEMENT_SHADER_TYPE(, FSlateElementVS, TEXT("SlateVertexShader"),TEXT("Main"),
 
 IMPLEMENT_SHADER_TYPE(, FSlateDebugOverdrawPS, TEXT("SlateElementPixelShader"), TEXT("DebugOverdrawMain"), SF_Pixel );
 
+IMPLEMENT_SHADER_TYPE(, FSlatePostProcessBlurPS, TEXT("SlatePostProcessPixelShader"), TEXT("GaussianBlurMain"), SF_Pixel);
+
+IMPLEMENT_SHADER_TYPE(, FSlatePostProcessDownsamplePS, TEXT("SlatePostProcessPixelShader"), TEXT("DownsampleMain"), SF_Pixel);
+
+
 #define IMPLEMENT_SLATE_PIXELSHADER_TYPE(ShaderType, bDrawDisabledEffect, bUseTextureAlpha) \
 	typedef TSlateElementPS<ESlateShader::ShaderType,bDrawDisabledEffect,bUseTextureAlpha> TSlateElementPS##ShaderType##bDrawDisabledEffect##bUseTextureAlpha##A; \
 	IMPLEMENT_SHADER_TYPE(template<>,TSlateElementPS##ShaderType##bDrawDisabledEffect##bUseTextureAlpha##A,TEXT("SlateElementPixelShader"),TEXT("Main"),SF_Pixel);
diff --git a/Engine/Source/Runtime/SlateRHIRenderer/Private/SlateShaders.h b/Engine/Source/Runtime/SlateRHIRenderer/Private/SlateShaders.h
--- a/Engine/Source/Runtime/SlateRHIRenderer/Private/SlateShaders.h
+++ b/Engine/Source/Runtime/SlateRHIRenderer/Private/SlateShaders.h
@@ -242,6 +242,102 @@ public:
 private:
 };
 
+const int32 MAX_BLUR_SAMPLES = 127;
+
+class FSlatePostProcessBlurPS : public FSlateElementPS
+{
+	DECLARE_SHADER_TYPE(FSlatePostProcessBlurPS, Global);
+public:
+	/** Indicates that this shader should be cached */
+	static bool ShouldCache(EShaderPlatform Platform)
+	{
+		return true;
+	}
+
+	FSlatePostProcessBlurPS()
+	{
+	}
+
+	/** Constructor.  Binds all parameters used by the shader */
+	FSlatePostProcessBlurPS(const ShaderMetaType::CompiledShaderInitializerType& Initializer)
+		: FSlateElementPS(Initializer)
+	{
+		BufferSizeAndDirection.Bind(Initializer.ParameterMap, TEXT("BufferSizeAndDirection"));
+		WeightAndOffsets.Bind(Initializer.ParameterMap, TEXT("WeightAndOffsets"));
+		SampleCount.Bind(Initializer.ParameterMap, TEXT("SampleCount"));
+		UVBounds.Bind(Initializer.ParameterMap, TEXT("UVBounds"));
+	}
+
+	void SetBufferSizeAndDirection(FRHICommandList& RHICmdList, const FVector2D& InBufferSize, const FVector2D& InDir)
+	{
+		SetShaderValue(RHICmdList, GetPixelShader(), BufferSizeAndDirection, FVector4(InBufferSize, InDir));
+	}
+
+	void SetWeightsAndOffsets(FRHICommandList& RHICmdList, const TArray<FVector4>& InWeightsAndOffsets, int32 NumSamples )
+	{
+		check(InWeightsAndOffsets.Num() <= MAX_BLUR_SAMPLES);
+		SetShaderValueArray<FPixelShaderRHIParamRef, FVector4>(RHICmdList, GetPixelShader(), WeightAndOffsets, InWeightsAndOffsets.GetData(), InWeightsAndOffsets.Num() );
+		SetShaderValue(RHICmdList, GetPixelShader(), SampleCount, NumSamples);
+	}
+
+	void SetUVBounds(FRHICommandList& RHICmdList, const FVector4& InUVBounds)
+	{
+		SetShaderValue(RHICmdList, GetPixelShader(), UVBounds, InUVBounds);
+	}
+
+	
+	virtual bool Serialize(FArchive& Ar) override
+	{
+		Ar << BufferSizeAndDirection;
+		Ar << WeightAndOffsets;
+		Ar << SampleCount;	
+		Ar << UVBounds;
+		return FSlateElementPS::Serialize(Ar);
+	}
+private:
+	FShaderParameter BufferSizeAndDirection;
+	FShaderParameter WeightAndOffsets;
+	FShaderParameter SampleCount;
+	FShaderParameter UVBounds;
+};
+
+
+class FSlatePostProcessDownsamplePS : public FSlateElementPS
+{
+	DECLARE_SHADER_TYPE(FSlatePostProcessDownsamplePS, Global);
+public:
+	/** Indicates that this shader should be cached */
+	static bool ShouldCache(EShaderPlatform Platform)
+	{
+		return true;
+	}
+
+	FSlatePostProcessDownsamplePS()
+	{
+	}
+
+	/** Constructor.  Binds all parameters used by the shader */
+	FSlatePostProcessDownsamplePS(const ShaderMetaType::CompiledShaderInitializerType& Initializer)
+		: FSlateElementPS(Initializer)
+	{
+		UVBounds.Bind(Initializer.ParameterMap, TEXT("UVBounds"));
+	}
+
+	void SetUVBounds(FRHICommandList& RHICmdList, const FVector4& InUVBounds)
+	{
+		SetShaderValue(RHICmdList, GetPixelShader(), UVBounds, InUVBounds);
+	}
+
+	virtual bool Serialize(FArchive& Ar) override
+	{
+		Ar << UVBounds;
+		return FSlateElementPS::Serialize(Ar);
+	}
+private:
+	FShaderParameter UVBounds;
+};
+
+
 /** The simple element vertex declaration. */
 extern TGlobalResource<FSlateVertexDeclaration> GSlateVertexDeclaration;
 
