Upped number of draw buffers used by Slate from 3->4 due to Paragon tests showing a number of "just missed" GPU frames when Slate flushes buffers

diff --git a/Engine/Source/Runtime/SlateRHIRenderer/Private/Slate3DRenderer.h b/Engine/Source/Runtime/SlateRHIRenderer/Private/Slate3DRenderer.h
--- a/Engine/Source/Runtime/SlateRHIRenderer/Private/Slate3DRenderer.h
+++ b/Engine/Source/Runtime/SlateRHIRenderer/Private/Slate3DRenderer.h
@@ -28,7 +28,7 @@ public:
 private:
 
 	/** Double buffered draw buffers so that the rendering thread can be rendering windows while the game thread is setting up for next frame */
-	static const int32 NUM_DRAW_BUFFERS = 3;
+	static const int32 NUM_DRAW_BUFFERS = 4;
 	FSlateDrawBuffer DrawBuffers[NUM_DRAW_BUFFERS];
 
 	/** The font services to use for rendering text */
