[PS4] Fixed an issue where a missing resource transition caused Orion to load with black artifacts.

diff --git a/Engine/Source/Runtime/SlateRHIRenderer/Private/SlateRHIRenderer.cpp b/Engine/Source/Runtime/SlateRHIRenderer/Private/SlateRHIRenderer.cpp
--- a/Engine/Source/Runtime/SlateRHIRenderer/Private/SlateRHIRenderer.cpp
+++ b/Engine/Source/Runtime/SlateRHIRenderer/Private/SlateRHIRenderer.cpp
@@ -703,9 +703,9 @@ void FSlateRHIRenderer::DrawWindow_RenderThread(FRHICommandListImmediate& RHICmd
 		RHICmdList.SetViewport(0, 0, 0, ViewportWidth, ViewportHeight, 0.0f);
 		RHICmdList.TransitionResource(EResourceTransitionAccess::EWritable, BackBuffer);
 		
-		if( ViewportInfo.bRequiresStencilTest )
+		if (ViewportInfo.bRequiresStencilTest)
 		{
-			check(IsValidRef( ViewportInfo.DepthStencil ));
+			check(IsValidRef(ViewportInfo.DepthStencil));
 
 			// Reset the backbuffer as our color render target and also set a depth stencil buffer
 			FRHIRenderTargetView ColorView(BackBuffer, 0, -1, bClear ? ERenderTargetLoadAction::EClear : ERenderTargetLoadAction::ELoad, ERenderTargetStoreAction::EStore);
@@ -865,6 +865,7 @@ void FSlateRHIRenderer::DrawWindow_RenderThread(FRHICommandListImmediate& RHICmd
 		{
 			GEngine->StereoRenderingDevice->RenderTexture_RenderThread(RHICmdList, RHICmdList.GetViewportBackBuffer(ViewportInfo.ViewportRHI), ViewportInfo.GetRenderTargetTexture());
 		}
+		RHICmdList.TransitionResource(EResourceTransitionAccess::EReadable, BackBuffer);
 	}
 
 	// Calculate renderthread time (excluding idle time).	
