Fix slate resources being destroyed when a dynamic image brush is destoyed even if the resource is shared.

diff --git a/Engine/Source/Runtime/SlateRHIRenderer/Private/SlateRHIResourceManager.cpp b/Engine/Source/Runtime/SlateRHIRenderer/Private/SlateRHIResourceManager.cpp
--- a/Engine/Source/Runtime/SlateRHIRenderer/Private/SlateRHIResourceManager.cpp
+++ b/Engine/Source/Runtime/SlateRHIRenderer/Private/SlateRHIResourceManager.cpp
@@ -883,7 +883,8 @@ void FSlateRHIResourceManager::ReleaseDynamicResource( const FSlateBrush& InBrus
 		{
 			TSharedPtr<FSlateDynamicTextureResource> TextureResource = DynamicResourceMap.GetDynamicTextureResource(ResourceName);
 
-			if( TextureResource.IsValid() )
+			// Only release the texture resource if it isn't shared by other handles
+			if( TextureResource.IsValid() && (!TextureResource->Proxy || TextureResource->Proxy->HandleData.IsUnique()))
 			{
 				// Release the rendering resource, its no longer being used
 				BeginReleaseResource(TextureResource->RHIRefTexture);
