Fix memory leak with slate dynamic image brushes

diff --git a/Engine/Source/Runtime/SlateRHIRenderer/Private/SlateRHIResourceManager.cpp b/Engine/Source/Runtime/SlateRHIRenderer/Private/SlateRHIResourceManager.cpp
--- a/Engine/Source/Runtime/SlateRHIRenderer/Private/SlateRHIResourceManager.cpp
+++ b/Engine/Source/Runtime/SlateRHIRenderer/Private/SlateRHIResourceManager.cpp
@@ -845,6 +845,9 @@ void FSlateRHIResourceManager::ReleaseDynamicResource( const FSlateBrush& InBrus
 	// Note: Only dynamically loaded or utexture brushes can be dynamically released
 	if( InBrush.HasUObject() || InBrush.IsDynamicallyLoaded() )
 	{
+		// Reset the rendering resource handle when our resource is being released
+		InBrush.ResourceHandle = FSlateResourceHandle();
+
 		FName ResourceName = InBrush.GetResourceName();
 
 		UObject* ResourceObject = InBrush.GetResourceObject();
@@ -884,7 +887,7 @@ void FSlateRHIResourceManager::ReleaseDynamicResource( const FSlateBrush& InBrus
 			TSharedPtr<FSlateDynamicTextureResource> TextureResource = DynamicResourceMap.GetDynamicTextureResource(ResourceName);
 
 			// Only release the texture resource if it isn't shared by other handles
-			if( TextureResource.IsValid() && (!TextureResource->Proxy || TextureResource->Proxy->HandleData.IsUnique()))
+			if (TextureResource.IsValid() && (!TextureResource->Proxy || TextureResource->Proxy->HandleData.IsUnique() || !TextureResource->Proxy->HandleData.IsValid()))
 			{
 				// Release the rendering resource, its no longer being used
 				BeginReleaseResource(TextureResource->RHIRefTexture);
