[LLM] main slate scene render now has SceneRendering tag, moving ~6MB out of RenderingThread

diff --git a/Engine/Source/Runtime/SlateRHIRenderer/Private/SlateRHIRenderer.cpp b/Engine/Source/Runtime/SlateRHIRenderer/Private/SlateRHIRenderer.cpp
--- a/Engine/Source/Runtime/SlateRHIRenderer/Private/SlateRHIRenderer.cpp
+++ b/Engine/Source/Runtime/SlateRHIRenderer/Private/SlateRHIRenderer.cpp
@@ -30,6 +30,7 @@
 #include "CommonRenderResources.h"
 #include "RenderTargetPool.h"
 #include "RendererUtils.h"
+#include "HAL/LowLevelMemTracker.h"
 
 DECLARE_CYCLE_STAT(TEXT("Slate RT: Rendering"), STAT_SlateRenderingRTTime, STATGROUP_Slate);
 DECLARE_CYCLE_STAT(TEXT("Slate RT: Create Batches"), STAT_SlateRTCreateBatches, STATGROUP_Slate);
@@ -647,6 +648,8 @@ static FAutoConsoleVariableRef CVarSlateWireframe(TEXT("Slate.ShowWireFrame"), S
 /** Draws windows from a FSlateDrawBuffer on the render thread */
 void FSlateRHIRenderer::DrawWindow_RenderThread(FRHICommandListImmediate& RHICmdList, FViewportInfo& ViewportInfo, FSlateWindowElementList& WindowElementList, const struct FSlateDrawWindowCommandParams& DrawCommandParams)
 {
+	LLM_SCOPE(ELLMTag::SceneRender);
+
 	static uint32 LastTimestamp = FPlatformTime::Cycles();
 	{
 		SCOPED_DRAW_EVENT(RHICmdList, SlateUI);
