Format SlateUI draw event with Window title.

diff --git a/Engine/Source/Runtime/SlateRHIRenderer/Private/SlateRHIRenderer.cpp b/Engine/Source/Runtime/SlateRHIRenderer/Private/SlateRHIRenderer.cpp
--- a/Engine/Source/Runtime/SlateRHIRenderer/Private/SlateRHIRenderer.cpp
+++ b/Engine/Source/Runtime/SlateRHIRenderer/Private/SlateRHIRenderer.cpp
@@ -85,6 +85,9 @@ struct FSlateDrawWindowCommandParams
 	FSlateRHIRenderer* Renderer;
 	FSlateWindowElementList* WindowElementList;
 	SWindow* Window;
+#if WANTS_DRAW_MESH_EVENTS
+	FString WindowTitle;
+#endif
 	float WorldTimeSeconds;
 	float DeltaTimeSeconds;
 	float RealTimeSeconds;
@@ -652,7 +655,7 @@ void FSlateRHIRenderer::DrawWindow_RenderThread(FRHICommandListImmediate& RHICmd
 
 	static uint32 LastTimestamp = FPlatformTime::Cycles();
 	{
-		SCOPED_DRAW_EVENT(RHICmdList, SlateUI);
+		SCOPED_DRAW_EVENTF(RHICmdList, SlateUI, TEXT("SlateUI Title = %s"), DrawCommandParams.WindowTitle.IsEmpty() ? TEXT("<none>") : *DrawCommandParams.WindowTitle);
 		SCOPED_GPU_STAT(RHICmdList, SlateUI);
 
 		// Should only be called by the rendering thread
@@ -1139,6 +1142,9 @@ void FSlateRHIRenderer::DrawWindows_Private(FSlateDrawBuffer& WindowDrawBuffer)
 					Params.Renderer = this;
 					Params.WindowElementList = &ElementList;
 					Params.Window = Window;
+#if WANTS_DRAW_MESH_EVENTS
+					Params.WindowTitle = Window->GetTitle().ToString();
+#endif
 					Params.bLockToVsync = bLockToVsync;
 #if ALPHA_BLENDED_WINDOWS
 					Params.bClear = Window->GetTransparencySupport() == EWindowTransparency::PerPixel;
