Temporary workaround for Slate rendering issue when switching from handheld to console mode after a DepthStencilTarget is created.

diff --git a/Engine/Source/Runtime/SlateRHIRenderer/Private/SlateRHIRenderer.cpp b/Engine/Source/Runtime/SlateRHIRenderer/Private/SlateRHIRenderer.cpp
--- a/Engine/Source/Runtime/SlateRHIRenderer/Private/SlateRHIRenderer.cpp
+++ b/Engine/Source/Runtime/SlateRHIRenderer/Private/SlateRHIRenderer.cpp
@@ -820,12 +820,14 @@ void FSlateRHIRenderer::DrawWindow_RenderThread(FRHICommandListImmediate& RHICmd
 						FSlateRenderingParams RenderParams(ViewMatrix * ViewportInfo.ProjectionMatrix, DrawCommandParams.WorldTimeSeconds, DrawCommandParams.DeltaTimeSeconds, DrawCommandParams.RealTimeSeconds);
 						RenderParams.bWireFrame = !!SlateWireFrame;
 
+						FTexture2DRHIRef EmptyTarget;
+
 						RenderingPolicy->DrawElements
 						(
 							RHICmdList,
 							BackBufferTarget,
 							BackBuffer,
-							ViewportInfo.DepthStencil,
+							ViewportInfo.bRequiresStencilTest ? ViewportInfo.DepthStencil : EmptyTarget,
 							BatchData.GetRenderBatches(),
 							RenderParams
 						);
