Fix for trying to use the backbuffer as readable and writable in the CompoisteUIPixelShade and a crash when no write mask support is available.

diff --git a/Engine/Source/Runtime/SlateRHIRenderer/Private/SlateRHIRenderer.cpp b/Engine/Source/Runtime/SlateRHIRenderer/Private/SlateRHIRenderer.cpp
--- a/Engine/Source/Runtime/SlateRHIRenderer/Private/SlateRHIRenderer.cpp
+++ b/Engine/Source/Runtime/SlateRHIRenderer/Private/SlateRHIRenderer.cpp
@@ -729,7 +729,8 @@ void FSlateRHIRenderer::DrawWindow_RenderThread(FRHICommandListImmediate& RHICmd
 			const uint32 ViewportHeight = (ViewportRT) ? ViewportRT->GetSizeY() : ViewportInfo.Height;
 
 			// Check to see that targets are up-to-date
-			if (bCompositeUI && (!ViewportInfo.UITargetRT || ViewportInfo.UITargetRT->GetRenderTargetItem().TargetableTexture->GetTexture2D()->GetSizeX() != ViewportWidth || ViewportInfo.UITargetRT->GetRenderTargetItem().TargetableTexture->GetTexture2D()->GetSizeY() != ViewportHeight))
+			if (bCompositeUI && (!ViewportInfo.UITargetRT || ViewportInfo.UITargetRT->GetRenderTargetItem().TargetableTexture->GetTexture2D()->GetSizeX() != ViewportWidth || ViewportInfo.UITargetRT->GetRenderTargetItem().TargetableTexture->GetTexture2D()->GetSizeY() != ViewportHeight
+				|| !ViewportInfo.HDRSourceRT || ViewportInfo.HDRSourceRT->GetRenderTargetItem().TargetableTexture->GetFormat() != BackBuffer->GetFormat()))
 			{
 				// Composition buffers
 				uint32 BaseFlags = (GSupportsRenderTargetWriteMask) ? TexCreate_NoFastClearFinalize : TexCreate_None;
@@ -746,6 +747,9 @@ void FSlateRHIRenderer::DrawWindow_RenderThread(FRHICommandListImmediate& RHICmd
 
 				GRenderTargetPool.FindFreeElement(RHICmdList, Desc, ViewportInfo.UITargetRT, TEXT("UITargetRT"));
 
+				Desc.Format = BackBuffer->GetFormat();
+				GRenderTargetPool.FindFreeElement(RHICmdList, Desc, ViewportInfo.HDRSourceRT, TEXT("HDRSourceRT"));
+
 				// LUT
 				ViewportInfo.ColorSpaceLUTRT.SafeRelease();
 				ViewportInfo.ColorSpaceLUTSRV.SafeRelease();
@@ -762,9 +766,9 @@ void FSlateRHIRenderer::DrawWindow_RenderThread(FRHICommandListImmediate& RHICmd
 			{
 				bClear = true; // Force a clear of the UI buffer to black
 
-				// Make sure the backbuffer is readable
+				// Grab HDR backbuffer
 				FResolveParams ResolveParams;
-				RHICmdList.CopyToResolveTarget(FinalBuffer, FinalBuffer, ResolveParams);
+				RHICmdList.CopyToResolveTarget(FinalBuffer, ViewportInfo.HDRSourceRT->GetRenderTargetItem().TargetableTexture, ResolveParams);
 
 				// UI backbuffer is temp target
 				BackBuffer = ViewportInfo.UITargetRT->GetRenderTargetItem().TargetableTexture->GetTexture2D();
@@ -918,6 +922,7 @@ void FSlateRHIRenderer::DrawWindow_RenderThread(FRHICommandListImmediate& RHICmd
 
 						TShaderMapRef<FScreenVS> VertexShader(ShaderMap);
 
+						FTextureRHIParamRef UITargetRTMaskTexture = GSupportsRenderTargetWriteMask ? ViewportInfo.UITargetRTMask->GetRenderTargetItem().TargetableTexture : nullptr;
 						if (HDROutputDevice == 5 || HDROutputDevice == 6)
 						{
 							// ScRGB encoding
@@ -930,7 +935,7 @@ void FSlateRHIRenderer::DrawWindow_RenderThread(FRHICommandListImmediate& RHICmd
 
 							SetGraphicsPipelineState(RHICmdList, GraphicsPSOInit);
 
-							PixelShader->SetParameters(RHICmdList, ViewportInfo.UITargetRT->GetRenderTargetItem().TargetableTexture, ViewportInfo.UITargetRTMask->GetRenderTargetItem().TargetableTexture, FinalBuffer, ViewportInfo.ColorSpaceLUTSRV);
+							PixelShader->SetParameters(RHICmdList, ViewportInfo.UITargetRT->GetRenderTargetItem().TargetableTexture, UITargetRTMaskTexture, ViewportInfo.HDRSourceRT->GetRenderTargetItem().TargetableTexture, ViewportInfo.ColorSpaceLUTSRV);
 						}
 						else
 						{
@@ -944,7 +949,7 @@ void FSlateRHIRenderer::DrawWindow_RenderThread(FRHICommandListImmediate& RHICmd
 
 							SetGraphicsPipelineState(RHICmdList, GraphicsPSOInit);
 
-							PixelShader->SetParameters(RHICmdList, ViewportInfo.UITargetRT->GetRenderTargetItem().TargetableTexture, ViewportInfo.UITargetRTMask->GetRenderTargetItem().TargetableTexture, FinalBuffer, ViewportInfo.ColorSpaceLUTSRV);
+							PixelShader->SetParameters(RHICmdList, ViewportInfo.UITargetRT->GetRenderTargetItem().TargetableTexture, UITargetRTMaskTexture, ViewportInfo.HDRSourceRT->GetRenderTargetItem().TargetableTexture, ViewportInfo.ColorSpaceLUTSRV);
 						}
 
 						RendererModule.DrawRectangle(
diff --git a/Engine/Source/Runtime/SlateRHIRenderer/Private/SlateRHIRenderer.h b/Engine/Source/Runtime/SlateRHIRenderer/Private/SlateRHIRenderer.h
--- a/Engine/Source/Runtime/SlateRHIRenderer/Private/SlateRHIRenderer.h
+++ b/Engine/Source/Runtime/SlateRHIRenderer/Private/SlateRHIRenderer.h
@@ -63,6 +63,8 @@ struct FViewportInfo : public FRenderResource
 	/** sRGB UI render target */
 	TRefCountPtr<IPooledRenderTarget> UITargetRT;
 	TRefCountPtr<IPooledRenderTarget> UITargetRTMask;
+	/** Copy of the HDR backbuffer */
+	TRefCountPtr<IPooledRenderTarget> HDRSourceRT;
 
 	/** Color-space LUT for HDR UI composition. */
 	FTexture3DRHIRef ColorSpaceLUTRT;
