Avoiding a crash when combining colorblind correction with retainers leading to an issue with render passes being stopped without it beginning.  This also fixes an issue with retainers being double colorblind corrected.

diff --git a/Engine/Source/Runtime/SlateRHIRenderer/Private/Slate3DRenderer.cpp b/Engine/Source/Runtime/SlateRHIRenderer/Private/Slate3DRenderer.cpp
--- a/Engine/Source/Runtime/SlateRHIRenderer/Private/Slate3DRenderer.cpp
+++ b/Engine/Source/Runtime/SlateRHIRenderer/Private/Slate3DRenderer.cpp
@@ -50,6 +50,11 @@ void FSlate3DRenderer::SetUseGammaCorrection(bool bUseGammaCorrection)
 	RenderTargetPolicy->SetUseGammaCorrection(bUseGammaCorrection);
 }
 
+void FSlate3DRenderer::SetApplyColorDeficiencyCorrection(bool bApplyColorCorrection)
+{
+	RenderTargetPolicy->SetApplyColorDeficiencyCorrection(bApplyColorCorrection);
+}
+
 FSlateDrawBuffer& FSlate3DRenderer::GetDrawBuffer()
 {
 	FreeBufferIndex = (FreeBufferIndex + 1) % NUM_DRAW_BUFFERS;
diff --git a/Engine/Source/Runtime/SlateRHIRenderer/Private/Slate3DRenderer.h b/Engine/Source/Runtime/SlateRHIRenderer/Private/Slate3DRenderer.h
--- a/Engine/Source/Runtime/SlateRHIRenderer/Private/Slate3DRenderer.h
+++ b/Engine/Source/Runtime/SlateRHIRenderer/Private/Slate3DRenderer.h
@@ -21,6 +21,8 @@ public:
 	void Cleanup();
 
 	virtual void SetUseGammaCorrection(bool bUseGammaCorrection) override;
+	virtual void SetApplyColorDeficiencyCorrection(bool bApplyColorCorrection) override;
+
 	virtual FSlateDrawBuffer& GetDrawBuffer() override;
 	virtual void DrawWindow_GameThread(FSlateDrawBuffer& DrawBuffer) override;
 	virtual void DrawWindowToTarget_RenderThread(FRHICommandListImmediate& RHICmdList, const struct FRenderThreadUpdateContext& Context) override;
diff --git a/Engine/Source/Runtime/SlateRHIRenderer/Private/SlateRHIRenderingPolicy.cpp b/Engine/Source/Runtime/SlateRHIRenderer/Private/SlateRHIRenderingPolicy.cpp
--- a/Engine/Source/Runtime/SlateRHIRenderer/Private/SlateRHIRenderingPolicy.cpp
+++ b/Engine/Source/Runtime/SlateRHIRenderer/Private/SlateRHIRenderingPolicy.cpp
@@ -72,6 +72,7 @@ FSlateRHIRenderingPolicy::FSlateRHIRenderingPolicy(TSharedRef<FSlateFontServices
 	, PostProcessor(new FSlatePostProcessor)
 	, ResourceManager(InResourceManager)
 	, bGammaCorrect(true)
+	, bApplyColorDeficiencyCorrection(true)
 	, InitialBufferSizeOverride(InitialBufferSize)
 	, LastDeviceProfile(nullptr)
 {
@@ -1287,7 +1288,7 @@ void FSlateRHIRenderingPolicy::DrawElements(
 
 	// Don't do color correction on iOS or Android, we don't have the GPU overhead for it.
 #if !(PLATFORM_IOS || PLATFORM_ANDROID)
-	if (GSlateColorDeficiencyType != EColorVisionDeficiency::NormalVision && GSlateColorDeficiencySeverity > 0)
+	if (bApplyColorDeficiencyCorrection && GSlateColorDeficiencyType != EColorVisionDeficiency::NormalVision && GSlateColorDeficiencySeverity > 0)
 	{
 		RHICmdList.EndRenderPass();
 
diff --git a/Engine/Source/Runtime/SlateRHIRenderer/Private/SlateRHIRenderingPolicy.h b/Engine/Source/Runtime/SlateRHIRenderer/Private/SlateRHIRenderingPolicy.h
--- a/Engine/Source/Runtime/SlateRHIRenderer/Private/SlateRHIRenderingPolicy.h
+++ b/Engine/Source/Runtime/SlateRHIRenderer/Private/SlateRHIRenderingPolicy.h
@@ -65,6 +65,7 @@ public:
 	void EndDrawingWindows();
 
 	void SetUseGammaCorrection( bool bInUseGammaCorrection ) { bGammaCorrect = bInUseGammaCorrection; }
+	void SetApplyColorDeficiencyCorrection(bool bInApplyColorCorrection) { bApplyColorDeficiencyCorrection = bInApplyColorCorrection; }
 
 	virtual void AddSceneAt(FSceneInterface* Scene, int32 Index) override;
 	virtual void ClearScenes() override;
@@ -104,6 +105,7 @@ private:
 	TSharedRef<FSlateRHIResourceManager> ResourceManager;
 
 	bool bGammaCorrect;
+	bool bApplyColorDeficiencyCorrection;
 
 	TOptional<int32> InitialBufferSizeOverride;
 
diff --git a/Engine/Source/Runtime/SlateRHIRenderer/Public/Interfaces/ISlate3DRenderer.h b/Engine/Source/Runtime/SlateRHIRenderer/Public/Interfaces/ISlate3DRenderer.h
--- a/Engine/Source/Runtime/SlateRHIRenderer/Public/Interfaces/ISlate3DRenderer.h
+++ b/Engine/Source/Runtime/SlateRHIRenderer/Public/Interfaces/ISlate3DRenderer.h
@@ -15,6 +15,7 @@ public:
 
 	/** set if this renderer should render in gamma space by default. */
 	virtual void SetUseGammaCorrection(bool bUseGammaCorrection) = 0;
+	virtual void SetApplyColorDeficiencyCorrection(bool bApplyColorCorrection) = 0;
 
 	/** @return The free buffer for drawing */
 	virtual FSlateDrawBuffer& GetDrawBuffer() = 0;
diff --git a/Engine/Source/Runtime/UMG/Private/Slate/SRetainerWidget.cpp b/Engine/Source/Runtime/UMG/Private/Slate/SRetainerWidget.cpp
--- a/Engine/Source/Runtime/UMG/Private/Slate/SRetainerWidget.cpp
+++ b/Engine/Source/Runtime/UMG/Private/Slate/SRetainerWidget.cpp
@@ -123,6 +123,10 @@ void SRetainerWidget::UpdateWidgetRenderer()
 	FWidgetRenderer* WidgetRenderer = RenderingResources->WidgetRenderer;
 
 	WidgetRenderer->SetUseGammaCorrection(bWriteContentInGammaSpace);
+
+	// This will be handled by the main slate rendering pass
+	WidgetRenderer->SetApplyColorDeficiencyCorrection(false);
+
 	WidgetRenderer->SetIsPrepassNeeded(false);
 	WidgetRenderer->SetClearHitTestGrid(false);
 
diff --git a/Engine/Source/Runtime/UMG/Private/Slate/WidgetRenderer.cpp b/Engine/Source/Runtime/UMG/Private/Slate/WidgetRenderer.cpp
--- a/Engine/Source/Runtime/UMG/Private/Slate/WidgetRenderer.cpp
+++ b/Engine/Source/Runtime/UMG/Private/Slate/WidgetRenderer.cpp
@@ -52,6 +52,16 @@ void FWidgetRenderer::SetUseGammaCorrection(bool bInUseGammaSpace)
 #endif
 }
 
+void FWidgetRenderer::SetApplyColorDeficiencyCorrection(bool bInApplyColorCorrection)
+{
+#if !UE_SERVER
+	if (LIKELY(FApp::CanEverRender()))
+	{
+		Renderer->SetApplyColorDeficiencyCorrection(bInApplyColorCorrection);
+	}
+#endif
+}
+
 UTextureRenderTarget2D* FWidgetRenderer::DrawWidget(const TSharedRef<SWidget>& Widget, FVector2D DrawSize)
 {
 	if ( LIKELY(FApp::CanEverRender()) )
diff --git a/Engine/Source/Runtime/UMG/Public/Slate/WidgetRenderer.h b/Engine/Source/Runtime/UMG/Public/Slate/WidgetRenderer.h
--- a/Engine/Source/Runtime/UMG/Public/Slate/WidgetRenderer.h
+++ b/Engine/Source/Runtime/UMG/Public/Slate/WidgetRenderer.h
@@ -41,6 +41,8 @@ public:
 	bool GetUseGammaCorrection() const { return bUseGammaSpace; }
 	void SetUseGammaCorrection(bool bInUseGammaSpace);
 
+	void SetApplyColorDeficiencyCorrection(bool bInApplyColorCorrection);
+
 	ISlate3DRenderer* GetSlateRenderer();
 
 	static UTextureRenderTarget2D* CreateTargetFor(FVector2D DrawSize, TextureFilter InFilter, bool bUseGammaCorrection);
