Avoid empty render-passes during UI rendering

diff --git a/Engine/Source/Runtime/SlateRHIRenderer/Private/SlateRHIRenderingPolicy.cpp b/Engine/Source/Runtime/SlateRHIRenderer/Private/SlateRHIRenderingPolicy.cpp
--- a/Engine/Source/Runtime/SlateRHIRenderer/Private/SlateRHIRenderingPolicy.cpp
+++ b/Engine/Source/Runtime/SlateRHIRenderer/Private/SlateRHIRenderingPolicy.cpp
@@ -720,7 +720,28 @@ void FSlateRHIRenderingPolicy::DrawElements(
 	*/
 	for (int32 BatchIndex = 0; BatchIndex < RenderBatches.Num(); ++BatchIndex)
 	{
-		check(RHICmdList.IsInsideRenderPass());
+		if (!RHICmdList.IsInsideRenderPass())
+		{
+			// Restart the renderpass since the CustomDrawer or post-process may have changed it in last iteration
+			FRHIRenderPassInfo RPInfo(BackBuffer.GetRenderTargetTexture(), ERenderTargetActions::Load_Store);
+			RPInfo.DepthStencilRenderTarget.DepthStencilTarget = DepthStencilTarget;
+			if (DepthStencilTarget)
+			{
+				RPInfo.DepthStencilRenderTarget.Action = EDepthStencilTargetActions::LoadDepthStencil_StoreDepthStencil;
+				RPInfo.DepthStencilRenderTarget.ExclusiveDepthStencil = FExclusiveDepthStencil::DepthWrite_StencilWrite;
+			}
+			else
+			{
+				RPInfo.DepthStencilRenderTarget.Action = EDepthStencilTargetActions::DontLoad_DontStore;
+				RPInfo.DepthStencilRenderTarget.ExclusiveDepthStencil = FExclusiveDepthStencil::DepthNop_StencilNop;
+			}
+			RHICmdList.BeginRenderPass(RPInfo, TEXT("RestartingSlateDrawElements"));
+
+			// Something may have messed with the viewport size so set it back to the full target.
+			RHICmdList.SetViewport(0, 0, 0, BackBuffer.GetSizeXY().X, BackBuffer.GetSizeXY().Y, 0.0f);
+			RHICmdList.SetStreamSource(0, VertexBuffer->VertexBufferRHI, 0);
+		}
+				
 #if WITH_SLATE_VISUALIZERS
 		FLinearColor BatchColor = FLinearColor(BatchColors.GetUnitVector());
 #endif
@@ -1227,21 +1248,7 @@ void FSlateRHIRenderingPolicy::DrawElements(
 				PostProcessor->BlurRect(RHICmdList, RendererModule, BlurParams, RectParams);
 
 				check(RHICmdList.IsOutsideRenderPass());
-				{
-					FRHIRenderPassInfo RPInfo(BackBuffer.GetRenderTargetTexture(), ERenderTargetActions::Load_Store);
-					RPInfo.DepthStencilRenderTarget.DepthStencilTarget = DepthStencilTarget;
-					if (DepthStencilTarget)
-					{
-						RPInfo.DepthStencilRenderTarget.Action = EDepthStencilTargetActions::LoadDepthStencil_StoreDepthStencil;
-						RPInfo.DepthStencilRenderTarget.ExclusiveDepthStencil = FExclusiveDepthStencil::DepthWrite_StencilWrite;
-					}
-					else
-					{
-						RPInfo.DepthStencilRenderTarget.Action = EDepthStencilTargetActions::DontLoad_DontStore;
-						RPInfo.DepthStencilRenderTarget.ExclusiveDepthStencil = FExclusiveDepthStencil::DepthNop_StencilNop;
-					}
-					RHICmdList.BeginRenderPass(RPInfo, TEXT("RestartingSlateDrawElements"));
-				}
+				// Render pass for slate elements will be restarted on a next loop iteration if any
 			}
 		}
 		else
@@ -1250,6 +1257,7 @@ void FSlateRHIRenderingPolicy::DrawElements(
 			if (CustomDrawer.IsValid())
 			{
 				// CustomDrawers will change the rendertarget. So we must close any outstanding renderpasses.
+				// Render pass for slate elements will be restarted on a next loop iteration if any
 				RHICmdList.EndRenderPass();	
 
 				SLATE_DRAW_EVENT(RHICmdList, CustomDrawer);
@@ -1263,25 +1271,6 @@ void FSlateRHIRenderingPolicy::DrawElements(
 
 				//We reset the maskingID here because otherwise the RT might not get re-set in the lines above see: if (bClearStencil || bForceStateChange)
 				MaskingID = 0;
-
-				// Restart the renderpass since the CustomDrawer will have drawn into its own.
-				FRHIRenderPassInfo RPInfo(ColorTarget, ERenderTargetActions::Load_Store);
-				RPInfo.DepthStencilRenderTarget.DepthStencilTarget = DepthStencilTarget;
-				if (DepthStencilTarget)
-				{
-					RPInfo.DepthStencilRenderTarget.Action = EDepthStencilTargetActions::LoadDepthStencil_StoreDepthStencil;
-					RPInfo.DepthStencilRenderTarget.ExclusiveDepthStencil = FExclusiveDepthStencil::DepthWrite_StencilWrite;
-				}
-				else
-				{
-					RPInfo.DepthStencilRenderTarget.Action = EDepthStencilTargetActions::DontLoad_DontStore;
-					RPInfo.DepthStencilRenderTarget.ExclusiveDepthStencil = FExclusiveDepthStencil::DepthNop_StencilNop;
-				}
-				RHICmdList.BeginRenderPass(RPInfo, TEXT("SlateDrawElements"));
-
-				// Something may have messed with the viewport size so set it back to the full target.
-				RHICmdList.SetViewport(0, 0, 0, BackBuffer.GetSizeXY().X, BackBuffer.GetSizeXY().Y, 0.0f);
-				RHICmdList.SetStreamSource(0, VertexBuffer->VertexBufferRHI, 0);
 			}
 		} // CustomDrawer
 	}
@@ -1290,7 +1279,10 @@ void FSlateRHIRenderingPolicy::DrawElements(
 #if !(PLATFORM_IOS || PLATFORM_ANDROID)
 	if (bApplyColorDeficiencyCorrection && GSlateColorDeficiencyType != EColorVisionDeficiency::NormalVision && GSlateColorDeficiencySeverity > 0)
 	{
-		RHICmdList.EndRenderPass();
+		if (RHICmdList.IsInsideRenderPass())
+		{
+			RHICmdList.EndRenderPass();
+		}
 
 		FPostProcessRectParams RectParams;
 		RectParams.SourceTexture = BackBuffer.GetRenderTargetTexture();
