Slate - Introducing a way to control the contrast of the UI when rendering, to help developers comply with CVAA, and general accessibility concerns it's valuable to be able to allow users to adjust the contrast of a game's UI, so any element affected by [REDACTED_PERSON], will also be affected by any contrast.  This setting does not enhance the contrast of the scene.

diff --git a/Engine/Shaders/Private/SlateElementPixelShader.usf b/Engine/Shaders/Private/SlateElementPixelShader.usf
--- a/Engine/Shaders/Private/SlateElementPixelShader.usf
+++ b/Engine/Shaders/Private/SlateElementPixelShader.usf
@@ -5,7 +5,7 @@
 #include "GammaCorrectionCommon.ush"
 #include "SlateShaderCommon.ush"
 
-/** Display gamma x:gamma curve adjustment, y:inverse gamma (1/GEngine->DisplayGamma), z:InvertAlpha */
+/** Display gamma x:gamma curve adjustment, y:inverse gamma (1/GEngine->DisplayGamma), z:InvertAlpha, w:Contrast */
 half4 GammaAndAlphaValues;
 /** Per element params */
 float4 ShaderParams;
@@ -14,6 +14,22 @@ float4 ShaderParams;
 Texture2D ElementTexture;
 SamplerState ElementTextureSampler;
 
+/**
+ * Applys a contrast about a midpoint squeezing the colors towards dark or light depending on
+ * the midpoint.  Additionally - it increases the opaqueness of colors to try and seperate them from
+ * the background.
+ */
+half4 ApplyContrast(half4 InColor)
+{
+    FLATTEN if( GammaAndAlphaValues.w != 1.0f )
+    {
+		float ContrastMidpoint = 0.25; //TODO We may want to make the midpoint configurable so that users can pick a different one, based on the overall saturation of their UI.
+	    InColor.rgb = saturate(((InColor.rgb - ContrastMidpoint) * GammaAndAlphaValues.w) + ContrastMidpoint);
+    }
+    
+    return InColor;
+}
+
 half3 GammaCorrect(half3 InColor)
 {
 	half3 CorrectedColor = InColor;
@@ -256,7 +272,10 @@ half4 Main( VertexToPixelInterpolants VIn ) : SV_Target0
 	OutColor.rgb = lerp( OutColor.rgb, Grayish, clamp( distance( OutColor.rgb, Grayish ), 0, .8)  );
 #endif
 
-	// gamma correct
+    // Contrast
+    OutColor = ApplyContrast(OutColor);
+
+	// Gamma Correct
 	OutColor.rgb = GammaCorrect(OutColor.rgb);
 
 #if !USE_MATERIALS
diff --git a/Engine/Source/Runtime/SlateCore/Private/SlateCoreClasses.cpp b/Engine/Source/Runtime/SlateCore/Private/SlateCoreClasses.cpp
--- a/Engine/Source/Runtime/SlateCore/Private/SlateCoreClasses.cpp
+++ b/Engine/Source/Runtime/SlateCore/Private/SlateCoreClasses.cpp
@@ -12,6 +12,15 @@ TAutoConsoleVariable<float> GlobalScrollAmount(
 	TEXT("How much to scroll for each click of the mouse wheel (in Slate Screen Units)."));
 
 
+float GSlateContrast = 1;
+
+FAutoConsoleVariableRef CVarSlateContrast(
+	TEXT("Slate.Contrast"),
+	GSlateContrast,
+	TEXT("The amount of contrast to apply to the UI.")
+);
+
+
 
 int32 GSlateLayoutCaching = 0;
 
@@ -32,5 +41,7 @@ FAutoConsoleVariableRef CVarSlateFastWidgetPath(
 	TEXT("Whether or not we enable fast widget pathing.  This mode relies on parent pointers to work correctly.")
 );
 
+
+
 FSlateWidgetStyle::FSlateWidgetStyle()
 { }
diff --git a/Engine/Source/Runtime/SlateCore/Public/Types/SlateConstants.h b/Engine/Source/Runtime/SlateCore/Public/Types/SlateConstants.h
--- a/Engine/Source/Runtime/SlateCore/Public/Types/SlateConstants.h
+++ b/Engine/Source/Runtime/SlateCore/Public/Types/SlateConstants.h
@@ -8,3 +8,6 @@
 /** How much to scroll for each click of the mouse wheel (in Slate Screen Units). */
 extern SLATECORE_API TAutoConsoleVariable<float> GlobalScrollAmount;
 inline float GetGlobalScrollAmount() { return GlobalScrollAmount.GetValueOnAnyThread(); }
+
+/**  */
+extern SLATECORE_API float GSlateContrast;
diff --git a/Engine/Source/Runtime/SlateRHIRenderer/Private/SlateMaterialShader.cpp b/Engine/Source/Runtime/SlateRHIRenderer/Private/SlateMaterialShader.cpp
--- a/Engine/Source/Runtime/SlateRHIRenderer/Private/SlateMaterialShader.cpp
+++ b/Engine/Source/Runtime/SlateRHIRenderer/Private/SlateMaterialShader.cpp
@@ -127,9 +127,9 @@ void FSlateMaterialShaderPS::SetAdditionalTexture( FRHICommandList& RHICmdList,
 	SetTextureParameter(RHICmdList, GetPixelShader(), AdditionalTextureParameter, TextureParameterSampler, SamplerState, InTexture );
 }
 
-void FSlateMaterialShaderPS::SetDisplayGamma(FRHICommandList& RHICmdList, float InDisplayGamma)
+void FSlateMaterialShaderPS::SetDisplayGammaAndContrast(FRHICommandList& RHICmdList, float InDisplayGamma, float InContrast)
 {
-	FVector4 InGammaValues(2.2f / InDisplayGamma, 1.0f/InDisplayGamma, 0.0f, 0.0f);
+	FVector4 InGammaValues(2.2f / InDisplayGamma, 1.0f / InDisplayGamma, 0.0f, InContrast);
 
 	SetShaderValue(RHICmdList, GetPixelShader(), GammaAndAlphaValues, InGammaValues);
 }
diff --git a/Engine/Source/Runtime/SlateRHIRenderer/Private/SlateMaterialShader.h b/Engine/Source/Runtime/SlateRHIRenderer/Private/SlateMaterialShader.h
--- a/Engine/Source/Runtime/SlateRHIRenderer/Private/SlateMaterialShader.h
+++ b/Engine/Source/Runtime/SlateRHIRenderer/Private/SlateMaterialShader.h
@@ -59,7 +59,7 @@ public:
 
 	void SetParameters(FRHICommandList& RHICmdList, const FSceneView& View, const FMaterialRenderProxy* MaterialRenderProxy, const FMaterial* Material, const FVector4& InShaderParams);
 
-	void SetDisplayGamma(FRHICommandList& RHICmdList, float InDisplayGamma);
+	void SetDisplayGammaAndContrast(FRHICommandList& RHICmdList, float InDisplayGamma, float InContrast);
 
 	void SetAdditionalTexture( FRHICommandList& RHICmdList, const FTextureRHIParamRef InTexture, const FSamplerStateRHIRef SamplerState );
 
diff --git a/Engine/Source/Runtime/SlateRHIRenderer/Private/SlateRHIRenderingPolicy.cpp b/Engine/Source/Runtime/SlateRHIRenderer/Private/SlateRHIRenderingPolicy.cpp
--- a/Engine/Source/Runtime/SlateRHIRenderer/Private/SlateRHIRenderingPolicy.cpp
+++ b/Engine/Source/Runtime/SlateRHIRenderer/Private/SlateRHIRenderingPolicy.cpp
@@ -629,6 +629,7 @@ void FSlateRHIRenderingPolicy::DrawElements(
 
 	const float EngineGamma = GEngine ? GEngine->GetDisplayGamma() : 2.2f;
 	const float DisplayGamma = bGammaCorrect ? EngineGamma : 1.0f;
+	const float DisplayContrast = GSlateContrast;
 
 	int32 ScissorClips = 0;
 	int32 StencilClips = 0;
@@ -1020,7 +1021,8 @@ void FSlateRHIRenderingPolicy::DrawElements(
 					PixelShader->SetTexture(RHICmdList, TextureRHI, SamplerState);
 					PixelShader->SetShaderParams(RHICmdList, ShaderParams.PixelParams);
 					const float FinalGamma = EnumHasAnyFlags(DrawFlags, ESlateBatchDrawFlag::ReverseGamma) ? (1.0f / EngineGamma) : EnumHasAnyFlags(DrawFlags, ESlateBatchDrawFlag::NoGamma) ? 1.0f : DisplayGamma;
-					PixelShader->SetDisplayGammaAndInvertAlpha(RHICmdList, FinalGamma, EnumHasAllFlags(DrawEffects, ESlateDrawEffect::InvertAlpha) ? 1.0f : 0.0f);
+					const float FinalContrast = EnumHasAnyFlags(DrawFlags, ESlateBatchDrawFlag::NoGamma) ? 1 : DisplayContrast;
+					PixelShader->SetDisplayGammaAndInvertAlphaAndContrast(RHICmdList, FinalGamma, EnumHasAllFlags(DrawEffects, ESlateDrawEffect::InvertAlpha) ? 1.0f : 0.0f, FinalContrast);
 				}
 
 				{
@@ -1143,7 +1145,8 @@ void FSlateRHIRenderingPolicy::DrawElements(
 
 								PixelShader->SetParameters(RHICmdList, ActiveSceneView, MaterialRenderProxy, Material, ShaderParams.PixelParams);
 								const float FinalGamma = EnumHasAnyFlags(DrawFlags, ESlateBatchDrawFlag::ReverseGamma) ? 1.0f / EngineGamma : EnumHasAnyFlags(DrawFlags, ESlateBatchDrawFlag::NoGamma) ? 1.0f : DisplayGamma;
-								PixelShader->SetDisplayGamma(RHICmdList, FinalGamma);
+								const float FinalContrast = EnumHasAnyFlags(DrawFlags, ESlateBatchDrawFlag::NoGamma) ? 1 : DisplayContrast;
+								PixelShader->SetDisplayGammaAndContrast(RHICmdList, FinalGamma, FinalContrast);
 
 								if (MaskResource)
 								{
diff --git a/Engine/Source/Runtime/SlateRHIRenderer/Private/SlateShaders.h b/Engine/Source/Runtime/SlateRHIRenderer/Private/SlateShaders.h
--- a/Engine/Source/Runtime/SlateRHIRenderer/Private/SlateShaders.h
+++ b/Engine/Source/Runtime/SlateRHIRenderer/Private/SlateShaders.h
@@ -157,9 +157,9 @@ public:
 	 *
 	 * @param DisplayGamma The display gamma to use
 	 */
-	void SetDisplayGammaAndInvertAlpha(FRHICommandList& RHICmdList, float InDisplayGamma, float bInvertAlpha)
+	void SetDisplayGammaAndInvertAlphaAndContrast(FRHICommandList& RHICmdList, float InDisplayGamma, float bInvertAlpha, float InContrast)
 	{
-		FVector4 Values( 2.2f / InDisplayGamma, 1.0f/InDisplayGamma, bInvertAlpha, 0.0f);
+		FVector4 Values( 2.2f / InDisplayGamma, 1.0f/InDisplayGamma, bInvertAlpha, InContrast);
 
 		SetShaderValue(RHICmdList, GetPixelShader(), GammaAndAlphaValues, Values);
 	}
