Fixing non-unity compile of the Contrast option.

diff --git a/Engine/Source/Runtime/SlateRHIRenderer/Private/SlateRHIRenderingPolicy.cpp b/Engine/Source/Runtime/SlateRHIRenderer/Private/SlateRHIRenderingPolicy.cpp
--- a/Engine/Source/Runtime/SlateRHIRenderer/Private/SlateRHIRenderingPolicy.cpp
+++ b/Engine/Source/Runtime/SlateRHIRenderer/Private/SlateRHIRenderingPolicy.cpp
@@ -25,6 +25,7 @@
 #include "Math/RandomStream.h"
 #include "DeviceProfiles/DeviceProfile.h"
 #include "DeviceProfiles/DeviceProfileManager.h"
+#include "Types/SlateConstants.h"
 
 extern void UpdateNoiseTextureParameters(FViewUniformShaderParameters& ViewUniformShaderParameters);
 
