Fix release of HDR related rendering resources on render thread.

diff --git a/Engine/Source/Runtime/SlateRHIRenderer/Private/SlateRHIRenderer.cpp b/Engine/Source/Runtime/SlateRHIRenderer/Private/SlateRHIRenderer.cpp
--- a/Engine/Source/Runtime/SlateRHIRenderer/Private/SlateRHIRenderer.cpp
+++ b/Engine/Source/Runtime/SlateRHIRenderer/Private/SlateRHIRenderer.cpp
@@ -108,6 +108,14 @@ void FViewportInfo::ReleaseRHI()
 	ViewportRHI.SafeRelease();
 }
 
+void FViewportInfo::ReleaseResource()
+{
+	FRenderResource::ReleaseResource();
+	UITargetRT.SafeRelease();
+	UITargetRTMask.SafeRelease();
+	HDRSourceRT.SafeRelease();
+}
+
 void FViewportInfo::ConditionallyUpdateDepthBuffer(bool bInRequiresStencilTest, uint32 InWidth, uint32 InHeight)
 {
 	FViewportInfo* ViewportInfo = this;
diff --git a/Engine/Source/Runtime/SlateRHIRenderer/Private/SlateRHIRenderer.h b/Engine/Source/Runtime/SlateRHIRenderer/Private/SlateRHIRenderer.h
--- a/Engine/Source/Runtime/SlateRHIRenderer/Private/SlateRHIRenderer.h
+++ b/Engine/Source/Runtime/SlateRHIRenderer/Private/SlateRHIRenderer.h
@@ -101,6 +101,7 @@ struct FViewportInfo : public FRenderResource
 	/** FRenderResource interface */
 	virtual void InitRHI() override;
 	virtual void ReleaseRHI() override;
+	virtual void ReleaseResource() override;
 
 	FViewportInfo()
 		:	ColorSpaceLUTOutputDevice(0),
