fix enginetest UI crash (more a workaround than a proper fix)

diff --git a/Engine/Source/Runtime/SlateRHIRenderer/Private/Slate3DRenderer.cpp b/Engine/Source/Runtime/SlateRHIRenderer/Private/Slate3DRenderer.cpp
--- a/Engine/Source/Runtime/SlateRHIRenderer/Private/Slate3DRenderer.cpp
+++ b/Engine/Source/Runtime/SlateRHIRenderer/Private/Slate3DRenderer.cpp
@@ -198,7 +198,10 @@ void FSlate3DRenderer::DrawWindowToTarget_RenderThread(FRHICommandListImmediate&
 			}
 		}
 	}
-	InRHICmdList.EndRenderPass();
+	if(InRHICmdList.IsInsideRenderPass())
+	{
+		InRHICmdList.EndRenderPass();
+	}
 
 	FSlateEndDrawingWindowsCommand::EndDrawingWindows(InRHICmdList, Context.WindowDrawBuffer, *RenderTargetPolicy);
 	InRHICmdList.CopyToResolveTarget(RenderTargetResource->GetRenderTargetTexture(), RTTextureRHI, FResolveParams());
