Fix window flicker from CL 6929492 and scroll bars showing when they shouldnt CL 6941620

diff --git a/Engine/Source/Runtime/SlateCore/Private/Widgets/SWindow.cpp b/Engine/Source/Runtime/SlateCore/Private/Widgets/SWindow.cpp
--- a/Engine/Source/Runtime/SlateCore/Private/Widgets/SWindow.cpp
+++ b/Engine/Source/Runtime/SlateCore/Private/Widgets/SWindow.cpp
@@ -883,7 +883,7 @@ void SWindow::ReshapeWindow( FVector2D NewPosition, FVector2D NewSize )
 	const FVector2D CurrentSize = GetSizeInScreen();
 
 	const FVector2D NewPositionTruncated = FVector2D(FMath::TruncToInt(NewPosition.X), FMath::TruncToInt(NewPosition.Y));
-	const FVector2D NewSizeRounded = FVector2D(FMath::TruncToInt(NewSize.X), FMath::TruncToInt(NewSize.Y));
+	const FVector2D NewSizeRounded = FVector2D(FMath::CeilToInt(NewSize.X), FMath::CeilToInt(NewSize.Y));
 
 	if ( CurrentPosition != NewPositionTruncated || CurrentSize != NewSizeRounded )
 	{
diff --git a/Engine/Source/Runtime/SlateRHIRenderer/Private/SlateRHIRenderer.cpp b/Engine/Source/Runtime/SlateRHIRenderer/Private/SlateRHIRenderer.cpp
--- a/Engine/Source/Runtime/SlateRHIRenderer/Private/SlateRHIRenderer.cpp
+++ b/Engine/Source/Runtime/SlateRHIRenderer/Private/SlateRHIRenderer.cpp
@@ -289,8 +289,8 @@ void FSlateRHIRenderer::CreateViewport(const TSharedRef<SWindow> Window)
 
 		// Clamp the window size to a reasonable default anything below 8 is a d3d warning and 8 is used anyway.
 		// @todo Slate: This is a hack to work around menus being summoned with 0,0 for window size until they are ticked.
-		int32 Width = FMath::Max(MIN_VIEWPORT_SIZE, FMath::TruncToInt(WindowSize.X));
-		int32 Height = FMath::Max(MIN_VIEWPORT_SIZE, FMath::TruncToInt(WindowSize.Y));
+		int32 Width = FMath::Max(MIN_VIEWPORT_SIZE, FMath::CeilToInt(WindowSize.X));
+		int32 Height = FMath::Max(MIN_VIEWPORT_SIZE, FMath::CeilToInt(WindowSize.Y));
 
 		// Sanity check dimensions
 		if (!ensureMsgf(Width <= MAX_VIEWPORT_SIZE && Height <= MAX_VIEWPORT_SIZE, TEXT("Invalid window with Width=%u and Height=%u"), Width, Height))
