https://jira.it.epicgames.net/browse/UE-65218

Basic issue is that FSlateRenderingPolicy::DrawElements sets a lambda which calls UpdateScissorRect. This is called later in FSlatePostProcessor::UpsampleRect. UpdateScissorRect should set and update the PSOInit but drops it on the floor.

diff --git a/Engine/Source/Runtime/SlateRHIRenderer/Private/SlateRHIRenderingPolicy.cpp b/Engine/Source/Runtime/SlateRHIRenderer/Private/SlateRHIRenderingPolicy.cpp
--- a/Engine/Source/Runtime/SlateRHIRenderer/Private/SlateRHIRenderingPolicy.cpp
+++ b/Engine/Source/Runtime/SlateRHIRenderer/Private/SlateRHIRenderingPolicy.cpp
@@ -1196,7 +1196,7 @@ void FSlateRHIRenderingPolicy::DrawElements(
 						LastClippingState,
 						ViewTranslation2D,
 						bSwitchVerticalAxis,
-						GraphicsPSOInit,
+						InGraphicsPSOInit,
 						Options.ViewProjectionMatrix,
 						true);
 				};
