Do not accumulate RHI thread stats if there is no RHI thread.

diff --git a/Engine/Source/Runtime/SlateRHIRenderer/Private/SlateRHIRenderer.cpp b/Engine/Source/Runtime/SlateRHIRenderer/Private/SlateRHIRenderer.cpp
--- a/Engine/Source/Runtime/SlateRHIRenderer/Private/SlateRHIRenderer.cpp
+++ b/Engine/Source/Runtime/SlateRHIRenderer/Private/SlateRHIRenderer.cpp
@@ -961,19 +961,22 @@ void FSlateRHIRenderer::DrawWindow_RenderThread(FRHICommandListImmediate& RHICmd
 
 	SET_CYCLE_COUNTER(STAT_RenderingIdleTime, RenderThreadIdle);
 	GRenderThreadTime = (ThreadTime > RenderThreadIdle) ? (ThreadTime - RenderThreadIdle) : ThreadTime;
-
-	RHICmdList.EnqueueLambda([](FRHICommandListImmediate&)
+	
+	if (IsRunningRHIInSeparateThread())
 	{
-		// Restart the RHI thread timer, so we don't count time spent in Present twice when this command list finishes.
-		int32 ThisCycles = FPlatformTime::Cycles();
-		GWorkingRHIThreadTime += (ThisCycles - GWorkingRHIThreadStartCycles);
-		GWorkingRHIThreadStartCycles = ThisCycles;
-
-		uint32 NewVal = GWorkingRHIThreadTime - GWorkingRHIThreadStallTime;
-		FPlatformAtomics::AtomicStore((int32*)&GRHIThreadTime, (int32)NewVal);
-		GWorkingRHIThreadTime = 0;
-		GWorkingRHIThreadStallTime = 0;
-	});
+		RHICmdList.EnqueueLambda([](FRHICommandListImmediate&)
+		{
+			// Restart the RHI thread timer, so we don't count time spent in Present twice when this command list finishes.
+			int32 ThisCycles = FPlatformTime::Cycles();
+			GWorkingRHIThreadTime += (ThisCycles - GWorkingRHIThreadStartCycles);
+			GWorkingRHIThreadStartCycles = ThisCycles;
+
+			uint32 NewVal = GWorkingRHIThreadTime - GWorkingRHIThreadStallTime;
+			FPlatformAtomics::AtomicStore((int32*)&GRHIThreadTime, (int32)NewVal);
+			GWorkingRHIThreadTime = 0;
+			GWorkingRHIThreadStallTime = 0;
+		});
+	}
 
 	// Any window that draws should be reporting referenced UObjects for the duration of the frame.
 	check(WindowElementList.ShouldReportUObjectReferences());
