Fix RHI thread timing
Exclude sleep time on RHI thread waiting for timestamp to be writter by GPU.

diff --git a/Engine/Source/Runtime/SlateRHIRenderer/Private/SlateRHIRenderer.cpp b/Engine/Source/Runtime/SlateRHIRenderer/Private/SlateRHIRenderer.cpp
--- a/Engine/Source/Runtime/SlateRHIRenderer/Private/SlateRHIRenderer.cpp
+++ b/Engine/Source/Runtime/SlateRHIRenderer/Private/SlateRHIRenderer.cpp
@@ -977,7 +977,7 @@ void FSlateRHIRenderer::DrawWindow_RenderThread(FRHICommandListImmediate& RHICmd
 		RHICmdList.EnqueueLambda([](FRHICommandListImmediate&)
 		{
 			// Restart the RHI thread timer, so we don't count time spent in Present twice when this command list finishes.
-			int32 ThisCycles = FPlatformTime::Cycles();
+			uint32 ThisCycles = FPlatformTime::Cycles();
 			GWorkingRHIThreadTime += (ThisCycles - GWorkingRHIThreadStartCycles);
 			GWorkingRHIThreadStartCycles = ThisCycles;
 
