- Backing out CL#4522771 to fix the sequencer rendering 4k output incorrectly
- This re-causes setting specific console variables to reset the window resolution to the previous window resolution

diff --git a/Engine/Source/Runtime/Slate/Private/Framework/Application/SlateApplication.cpp b/Engine/Source/Runtime/Slate/Private/Framework/Application/SlateApplication.cpp
--- a/Engine/Source/Runtime/Slate/Private/Framework/Application/SlateApplication.cpp
+++ b/Engine/Source/Runtime/Slate/Private/Framework/Application/SlateApplication.cpp
@@ -6795,7 +6795,10 @@ bool FSlateApplication::OnSizeChanged( const TSharedRef< FGenericWindow >& Platf
 
 		Renderer->RequestResize( Window, Width, Height );
 
-		Renderer->SetSystemResolution(Width, Height);
+		if (FPlatformProperties::HasFixedResolution())
+		{
+			Renderer->SetSystemResolution(Width, Height);
+		}
 		
 		if ( !bWasMinimized && Window->IsRegularWindow() && !Window->HasOSWindowBorder() && Window->IsVisible() && Window->IsDrawingEnabled() )
 		{
diff --git a/Engine/Source/Runtime/SlateRHIRenderer/Private/SlateRHIRenderer.cpp b/Engine/Source/Runtime/SlateRHIRenderer/Private/SlateRHIRenderer.cpp
--- a/Engine/Source/Runtime/SlateRHIRenderer/Private/SlateRHIRenderer.cpp
+++ b/Engine/Source/Runtime/SlateRHIRenderer/Private/SlateRHIRenderer.cpp
@@ -425,10 +425,7 @@ void FSlateRHIRenderer::UpdateFullscreenState(const TSharedRef<SWindow> Window,
 void FSlateRHIRenderer::SetSystemResolution(uint32 Width, uint32 Height)
 {
 	FSystemResolution::RequestResolutionChange(Width, Height, FPlatformProperties::HasFixedResolution() ? EWindowMode::Fullscreen : GSystemResolution.WindowMode);
-	if (FPlatformProperties::HasFixedResolution())
-	{
-		IConsoleManager::Get().CallAllConsoleVariableSinks();
-	}
+	IConsoleManager::Get().CallAllConsoleVariableSinks();
 }
 
 void FSlateRHIRenderer::RestoreSystemResolution(const TSharedRef<SWindow> InWindow)
