Moved slate stencil buffer clipping vertex layout to the vertex shader instead of creating a vertex buffer for each mask every frame.  Fixed masking not working on vulkan

diff --git a/Engine/Shaders/Private/SlateMaskingShader.usf b/Engine/Shaders/Private/SlateMaskingShader.usf
--- a/Engine/Shaders/Private/SlateMaskingShader.usf
+++ b/Engine/Shaders/Private/SlateMaskingShader.usf
@@ -4,6 +4,13 @@
 
 float4x4 ViewProjection;
 
+
+// MaskRect[0].xy = TopLeft
+// MaskRect[0].zw = TopRight
+// MaskRect[1].xy = BottomLeft
+// MaskRect[1].zw = BottomRight
+float4 MaskRectPacked[2];
+
 float SwitchVerticalAxisMultiplier;
 
 struct VertexToPixelInterpolants
@@ -11,15 +18,15 @@ struct VertexToPixelInterpolants
     float4 Position : SV_POSITION;
 };
 
-VertexToPixelInterpolants MainVS(
-	in float2 InPosition : ATTRIBUTE0
-	)
+VertexToPixelInterpolants MainVS(in uint InPosition : ATTRIBUTE0)
 {
 	VertexToPixelInterpolants VOut = (VertexToPixelInterpolants)0;
 
-    float4 WorldPosition = float4(InPosition.xy, 0, 1);
+	const float2 MaskRect[4] = { MaskRectPacked[0].xy, MaskRectPacked[0].zw, MaskRectPacked[1].xy, MaskRectPacked[1].zw };
+
+    float4 WorldPosition = float4(MaskRect[InPosition], 0, 1);
 
-    VOut.Position = mul(WorldPosition, ViewProjection);
+	VOut.Position = mul(WorldPosition, ViewProjection);
 
 #if (ES2_PROFILE && COMPILER_GLSL_ES2) || (ES3_1_PROFILE && COMPILER_GLSL_ES3_1)
 	// @todo-mobile: Fix this in the projection matrix
diff --git a/Engine/Source/Runtime/SlateRHIRenderer/Private/SlateElementVertexBuffer.h b/Engine/Source/Runtime/SlateRHIRenderer/Private/SlateElementVertexBuffer.h
--- a/Engine/Source/Runtime/SlateRHIRenderer/Private/SlateElementVertexBuffer.h
+++ b/Engine/Source/Runtime/SlateRHIRenderer/Private/SlateElementVertexBuffer.h
@@ -8,6 +8,7 @@
 #include "SlateGlobals.h"
 #include "RHI.h"
 #include "RenderResource.h"
+#include "Containers/ResourceArray.h"
 
 DECLARE_MEMORY_STAT_EXTERN(TEXT("Vertex Buffer Memory (GPU)"), STAT_SlateVertexBufferMemory, STATGROUP_SlateMemory, );
 
@@ -169,3 +170,58 @@ private:
 	TSlateElementVertexBuffer( const TSlateElementVertexBuffer& );
 	void operator=(const TSlateElementVertexBuffer& );
 };
+
+
+class FSlateStencilClipVertexBuffer : public FVertexBuffer
+{
+public:
+	FSlateStencilClipVertexBuffer()
+	{}
+
+	~FSlateStencilClipVertexBuffer() {};
+
+	/** Initializes the vertex buffers RHI resource. */
+	virtual void InitDynamicRHI()
+	{
+		if (!IsValidRef(VertexBufferRHI))
+		{
+			static FStencilBufferResourceArray ResourceArray;
+
+			FRHIResourceCreateInfo CreateInfo;
+			CreateInfo.ResourceArray = &ResourceArray;
+			VertexBufferRHI = RHICreateVertexBuffer(ResourceArray.GetResourceDataSize(), BUF_Static, CreateInfo);
+
+			// Ensure the vertex buffer could be created
+			check(IsValidRef(VertexBufferRHI));
+		}
+	}
+
+	/** Releases the vertex buffers RHI resource. */
+	virtual void ReleaseDynamicRHI()
+	{
+		VertexBufferRHI.SafeRelease();
+	}
+
+	/** Returns a friendly name for this buffer. */
+	virtual FString GetFriendlyName() const { return TEXT("SlateElementVertices"); }
+
+private:
+	struct FStencilBufferResourceArray : FResourceArrayInterface
+	{
+		virtual const void* GetResourceData() const override
+		{
+			static uint32 Verts[] = { 0, 1, 2, 3};
+			return Verts;
+		}
+
+		virtual uint32 GetResourceDataSize() const override
+		{
+			return sizeof(uint32) * 4;
+		}
+
+		virtual void Discard() override {}
+		virtual bool IsStatic() const override { return true; }
+		virtual bool GetAllowCPUAccess() const override { return false; }
+		virtual void SetAllowCPUAccess(bool bInNeedsCPUAccess) override { }
+	};
+};
diff --git a/Engine/Source/Runtime/SlateRHIRenderer/Private/SlateRHIRenderingPolicy.cpp b/Engine/Source/Runtime/SlateRHIRenderer/Private/SlateRHIRenderingPolicy.cpp
--- a/Engine/Source/Runtime/SlateRHIRenderer/Private/SlateRHIRenderingPolicy.cpp
+++ b/Engine/Source/Runtime/SlateRHIRenderer/Private/SlateRHIRenderingPolicy.cpp
@@ -42,14 +42,14 @@ DECLARE_CYCLE_STAT(TEXT("Slate RT: Custom Draw"), STAT_SlateRTCustomDraw, STATGR
 DECLARE_DWORD_COUNTER_STAT(TEXT("Clips (Scissor)"), STAT_SlateScissorClips, STATGROUP_Slate);
 DECLARE_DWORD_COUNTER_STAT(TEXT("Clips (Stencil)"), STAT_SlateStencilClips, STATGROUP_Slate);
 
-#if UE_BUILD_DEBUG
+#if WITH_SLATE_DEBUGGING
 int32 SlateEnableDrawEvents = 1;
 #else
 int32 SlateEnableDrawEvents = 0;
 #endif
 static FAutoConsoleVariableRef CVarSlateEnableDrawEvents(TEXT("Slate.EnableDrawEvents"), SlateEnableDrawEvents, TEXT("."), ECVF_Default);
 
-#if !(UE_BUILD_SHIPPING || UE_BUILD_TEST)
+#if WITH_SLATE_DEBUGGING
 int32 BatchToDraw = -1;
 static FAutoConsoleVariableRef CVarSlateDrawBatchNum(TEXT("Slate.DrawBatchNum"), BatchToDraw, TEXT("."), ECVF_Default);
 #endif
@@ -106,12 +106,16 @@ void FSlateRHIRenderingPolicy::InitResources()
 
 	VertexBuffers.Init(NumVertices);
 	IndexBuffers.Init(NumVertices);
+
+	BeginInitResource(&StencilVertexBuffer);
 }
 
 void FSlateRHIRenderingPolicy::ReleaseResources()
 {
 	VertexBuffers.Destroy();
 	IndexBuffers.Destroy();
+
+	BeginReleaseResource(&StencilVertexBuffer);
 }
 
 void FSlateRHIRenderingPolicy::BeginDrawingWindows()
@@ -318,6 +322,7 @@ static bool UpdateScissorRect(
 	const FVector2D& ViewTranslation2D, 
 	bool bSwitchVerticalAxis,
 	FGraphicsPipelineStateInitializer& InGraphicsPSOInit,
+	FSlateStencilClipVertexBuffer& StencilVertexBuffer,
 	const FMatrix& ViewProjection, 
 	bool bForceStateChange)
 {
@@ -444,6 +449,12 @@ static bool UpdateScissorRect(
 					RHICmdList.BeginRenderPass(RPInfo, TEXT("SlateUpdateScissorRect_ClearStencil"));
 				}
 
+				TShaderMap<FGlobalShaderType>* MaxFeatureLevelShaderMap = GetGlobalShaderMap(GMaxRHIFeatureLevel);
+
+				// Set the new shaders
+				TShaderMapRef<FSlateMaskingVS> VertexShader(MaxFeatureLevelShaderMap);
+				TShaderMapRef<FSlateMaskingPS> PixelShader(MaxFeatureLevelShaderMap);
+
 				// Start by setting up the stenciling states so that we can write representations of the clipping zones into the stencil buffer only.
 				{
 					FGraphicsPipelineStateInitializer WriteMaskPSOInit;
@@ -467,12 +478,6 @@ static bool UpdateScissorRect(
 						, /*StencilReadMask*/ 0xFF
 						, /*StencilWriteMask*/ 0xFF>::GetRHI();
 
-					TShaderMap<FGlobalShaderType>* MaxFeatureLevelShaderMap = GetGlobalShaderMap(GMaxRHIFeatureLevel);
-
-					// Set the new shaders
-					TShaderMapRef<FSlateMaskingVS> VertexShader(MaxFeatureLevelShaderMap);
-					TShaderMapRef<FSlateMaskingPS> PixelShader(MaxFeatureLevelShaderMap);
-
 					WriteMaskPSOInit.BoundShaderState.VertexDeclarationRHI = GSlateMaskingVertexDeclaration.VertexDeclarationRHI;
 					WriteMaskPSOInit.BoundShaderState.VertexShaderRHI = GETSAFERHISHADER_VERTEX(*VertexShader);
 					WriteMaskPSOInit.BoundShaderState.PixelShaderRHI = GETSAFERHISHADER_PIXEL(*PixelShader);
@@ -491,19 +496,9 @@ static bool UpdateScissorRect(
 
 						SCOPE_CYCLE_COUNTER(STAT_SlateRTStencilDrawCall);
 
-						FRHIResourceCreateInfo CreateInfo;
-						FVertexBufferRHIRef VertexBufferRHI = RHICreateVertexBuffer(sizeof(FVector2D) * 4, BUF_Volatile, CreateInfo);
-						void* VoidPtr = RHILockVertexBuffer(VertexBufferRHI, 0, sizeof(FVector2D) * 4, RLM_WriteOnly);
-						//TODO Slate If we ever decided to add masking with a texture, we could do that here.
-						FVector2D* Vertices = (FVector2D*)VoidPtr;
+						VertexShader->SetMaskRect(RHICmdList, MaskQuad.TopLeft, MaskQuad.TopRight, MaskQuad.BottomLeft, MaskQuad.BottomRight);
 
-						Vertices[0].Set(MaskQuad.TopLeft.X, MaskQuad.TopLeft.Y);
-						Vertices[1].Set(MaskQuad.TopRight.X, MaskQuad.TopRight.Y);
-						Vertices[2].Set(MaskQuad.BottomLeft.X, MaskQuad.BottomLeft.Y);
-						Vertices[3].Set(MaskQuad.BottomRight.X, MaskQuad.BottomRight.Y);
-
-						RHIUnlockVertexBuffer(VertexBufferRHI);
-						RHICmdList.SetStreamSource(0, VertexBufferRHI, 0);
+						RHICmdList.SetStreamSource(0, StencilVertexBuffer.VertexBufferRHI, 0);
 						RHICmdList.DrawPrimitive(0, 2, 1);
 					}
 
@@ -545,20 +540,10 @@ static bool UpdateScissorRect(
 					const FSlateClippingZone& MaskQuad = StencilQuads[MaskIndex];
 
 					SCOPE_CYCLE_COUNTER(STAT_SlateRTStencilDrawCall);
-					//TODO Slate If we ever decided to add masking with a texture, we could do that here.
-					FRHIResourceCreateInfo CreateInfo;
-					FVertexBufferRHIRef VertexBufferRHI = RHICreateVertexBuffer(sizeof(FVector2D) * 4, BUF_Volatile, CreateInfo);
-					void* VoidPtr = RHILockVertexBuffer(VertexBufferRHI, 0, sizeof(FVector2D) * 4, RLM_WriteOnly);
-					// Generate the vertices used
-					FVector2D* Vertices = (FVector2D*)VoidPtr;
-
-					Vertices[0].Set(MaskQuad.TopLeft.X, MaskQuad.TopLeft.Y);
-					Vertices[1].Set(MaskQuad.TopRight.X, MaskQuad.TopRight.Y);
-					Vertices[2].Set(MaskQuad.BottomLeft.X, MaskQuad.BottomLeft.Y);
-					Vertices[3].Set(MaskQuad.BottomRight.X, MaskQuad.BottomRight.Y);
-
-					RHIUnlockVertexBuffer(VertexBufferRHI);
-					RHICmdList.SetStreamSource(0, VertexBufferRHI, 0);
+					
+					VertexShader->SetMaskRect(RHICmdList, MaskQuad.TopLeft, MaskQuad.TopRight, MaskQuad.BottomLeft, MaskQuad.BottomRight);
+
+					RHICmdList.SetStreamSource(0, StencilVertexBuffer.VertexBufferRHI, 0);
 					RHICmdList.DrawPrimitive(0, 2, 1);
 				}
 
@@ -815,6 +800,7 @@ void FSlateRHIRenderingPolicy::DrawElements(
 				ViewTranslation2D,
 				bSwitchVerticalAxis,
 				GraphicsPSOInit,
+				StencilVertexBuffer,
 				ViewProjection,
 				false);
 
@@ -1225,6 +1211,7 @@ void FSlateRHIRenderingPolicy::DrawElements(
 						ViewTranslation2D,
 						bSwitchVerticalAxis,
 						InGraphicsPSOInit,
+						StencilVertexBuffer,
 						Options.ViewProjectionMatrix,
 						true);
 				};
diff --git a/Engine/Source/Runtime/SlateRHIRenderer/Private/SlateRHIRenderingPolicy.h b/Engine/Source/Runtime/SlateRHIRenderer/Private/SlateRHIRenderingPolicy.h
--- a/Engine/Source/Runtime/SlateRHIRenderer/Private/SlateRHIRenderingPolicy.h
+++ b/Engine/Source/Runtime/SlateRHIRenderer/Private/SlateRHIRenderingPolicy.h
@@ -90,6 +90,8 @@ private:
 	TSlateElementVertexBuffer<FSlateVertex> VertexBuffers;
 	FSlateElementIndexBuffer IndexBuffers;
 
+	FSlateStencilClipVertexBuffer StencilVertexBuffer;
+
 	/** Handles post process effects for slate */
 	TSharedRef<FSlatePostProcessor> PostProcessor;
 
diff --git a/Engine/Source/Runtime/SlateRHIRenderer/Private/SlateShaders.cpp b/Engine/Source/Runtime/SlateRHIRenderer/Private/SlateShaders.cpp
--- a/Engine/Source/Runtime/SlateRHIRenderer/Private/SlateShaders.cpp
+++ b/Engine/Source/Runtime/SlateRHIRenderer/Private/SlateShaders.cpp
@@ -100,8 +100,8 @@ void FSlateElementPS::ModifyCompilationEnvironment(const FGlobalShaderPermutatio
 void FSlateMaskingVertexDeclaration::InitRHI()
 {
 	FVertexDeclarationElementList Elements;
-	uint32 Stride = sizeof(FVector2D);
-	Elements.Add(FVertexElement(0, 0, VET_Float2, 0, Stride));
+	uint32 Stride = sizeof(uint32);
+	Elements.Add(FVertexElement(0, 0, VET_UByte4, 0, Stride));
 
 	VertexDeclarationRHI = RHICreateVertexDeclaration(Elements);
 }
@@ -160,6 +160,7 @@ FSlateMaskingVS::FSlateMaskingVS(const ShaderMetaType::CompiledShaderInitializer
 	: FGlobalShader(Initializer)
 {
 	ViewProjection.Bind(Initializer.ParameterMap, TEXT("ViewProjection"));
+	MaskRect.Bind(Initializer.ParameterMap, TEXT("MaskRectPacked"));
 	SwitchVerticalAxisMultiplier.Bind(Initializer.ParameterMap, TEXT("SwitchVerticalAxisMultiplier"));
 }
 
@@ -173,12 +174,21 @@ void FSlateMaskingVS::SetVerticalAxisMultiplier(FRHICommandList& RHICmdList, flo
 	SetShaderValue(RHICmdList, GetVertexShader(), SwitchVerticalAxisMultiplier, InMultiplier );
 }
 
+void FSlateMaskingVS::SetMaskRect(FRHICommandList& RHICmdList, const FVector2D& TopLeft, const FVector2D& TopRight, const FVector2D& BotLeft, const FVector2D& BotRight)
+{
+	//FVector4 MaskRectVal[4] = { FVector4(TopLeft, FVector2D::ZeroVector), FVector4(TopRight, FVector2D::ZeroVector), FVector4(BotLeft, FVector2D::ZeroVector), FVector4(BotRight, FVector2D::ZeroVector) };
+	FVector4 MaskRectVal[2] = { FVector4(TopLeft, TopRight), FVector4(BotLeft, BotRight) };
+
+	SetShaderValue(RHICmdList, GetVertexShader(), MaskRect, MaskRectVal);
+}
+
 /** Serializes the shader data */
 bool FSlateMaskingVS::Serialize(FArchive& Ar)
 {
 	bool bShaderHasOutdatedParameters = FGlobalShader::Serialize(Ar);
 
 	Ar << ViewProjection;
+	Ar << MaskRect;
 	Ar << SwitchVerticalAxisMultiplier;
 
 	return bShaderHasOutdatedParameters;
diff --git a/Engine/Source/Runtime/SlateRHIRenderer/Private/SlateShaders.h b/Engine/Source/Runtime/SlateRHIRenderer/Private/SlateShaders.h
--- a/Engine/Source/Runtime/SlateRHIRenderer/Private/SlateShaders.h
+++ b/Engine/Source/Runtime/SlateRHIRenderer/Private/SlateShaders.h
@@ -481,10 +481,16 @@ public:
 	 */
 	void SetVerticalAxisMultiplier(FRHICommandList& RHICmdList, float InMultiplier);
 
+	/**
+	 * Sets the mask rect positions
+	 */
+	void SetMaskRect(FRHICommandList& RHICmdList, const FVector2D& TopLeft, const FVector2D& TopRight, const FVector2D& BotLeft, const FVector2D& BotRight);
+
 	virtual bool Serialize(FArchive& Ar) override;
 
 private:
-
+	/** Mask rect parameter */
+	FShaderParameter MaskRect;
 	/** ViewProjection parameter used by the shader */
 	FShaderParameter ViewProjection;
 	/** Parameter used to determine if we need to swtich the vertical axis for opengl */
