Fix race condition with slates usage of time for materials.  Time was previously accessed directly on the render thread.  It is now passed from the game thread to the render thread like scene rendering.

diff --git a/Engine/Source/Runtime/SlateCore/Public/Rendering/DrawElements.h b/Engine/Source/Runtime/SlateCore/Public/Rendering/DrawElements.h
--- a/Engine/Source/Runtime/SlateCore/Public/Rendering/DrawElements.h
+++ b/Engine/Source/Runtime/SlateCore/Public/Rendering/DrawElements.h
@@ -1444,7 +1444,7 @@ public:
 /**
  * Represents a top level window and its draw elements.
  */
-class FSlateWindowElementList
+class FSlateWindowElementList : public FNoncopyable
 {
 	friend class FSlateElementBatcher;
 	friend class FSlateDrawElement;
diff --git a/Engine/Source/Runtime/SlateCore/Public/Rendering/SlateRenderer.h b/Engine/Source/Runtime/SlateCore/Public/Rendering/SlateRenderer.h
--- a/Engine/Source/Runtime/SlateCore/Public/Rendering/SlateRenderer.h
+++ b/Engine/Source/Runtime/SlateCore/Public/Rendering/SlateRenderer.h
@@ -27,7 +27,10 @@ typedef TRefCountPtr<FRHITexture2D> FTexture2DRHIRef;
  */
 struct FRenderThreadUpdateContext
 {
-	class FSlateDrawBuffer* DrawBuffer;
+	class FSlateDrawBuffer* WindowDrawBuffer;
+	float WorldTimeSeconds;
+	float DeltaTimeSeconds;
+	float RealTimeSeconds;
 	void* RenderTargetResource;
 	void* Renderer;
 	bool bClearTarget;
diff --git a/Engine/Source/Runtime/SlateRHIRenderer/Private/Slate3DRenderer.cpp b/Engine/Source/Runtime/SlateRHIRenderer/Private/Slate3DRenderer.cpp
--- a/Engine/Source/Runtime/SlateRHIRenderer/Private/Slate3DRenderer.cpp
+++ b/Engine/Source/Runtime/SlateRHIRenderer/Private/Slate3DRenderer.cpp
@@ -112,27 +112,27 @@ struct TKeepAliveCommand final : public FRHICommand < TKeepAliveCommand<TKeepAli
 	void Execute(FRHICommandListBase& CmdList) {}
 };
 
-void FSlate3DRenderer::DrawWindowToTarget_RenderThread( FRHICommandListImmediate& InRHICmdList, FTextureRenderTarget2DResource* RenderTargetResource, FSlateDrawBuffer& WindowDrawBuffer, bool bInClearTarget)
+void FSlate3DRenderer::DrawWindowToTarget_RenderThread(FRHICommandListImmediate& InRHICmdList, const FRenderThreadUpdateContext& Context)
 {
 	check(IsInRenderingThread());
 	QUICK_SCOPE_CYCLE_COUNTER(Stat_Slate_WidgetRendererRenderThread);
 	SCOPED_DRAW_EVENT( InRHICmdList, SlateRenderToTarget );
 	SCOPED_GPU_STAT(InRHICmdList, Slate3D);
 
-	checkSlow( RenderTargetResource );
+	checkSlow( Context.RenderTargetResource );
 
-	const TArray<TSharedRef<FSlateWindowElementList>>& WindowsToDraw = WindowDrawBuffer.GetWindowElementLists();
+	const TArray<TSharedRef<FSlateWindowElementList>>& WindowsToDraw = Context.WindowDrawBuffer->GetWindowElementLists();
 
 	// Enqueue a command to unlock the draw buffer after all windows have been drawn
 	RenderTargetPolicy->BeginDrawingWindows();
 
-
+	FTextureRenderTarget2DResource* RenderTargetResource = static_cast<FTextureRenderTarget2DResource*>(Context.RenderTargetResource);
 	// Set render target and clear.
-	FTexture2DRHIRef RTResource = RenderTargetResource->GetTextureRHI();
-	InRHICmdList.TransitionResource(EResourceTransitionAccess::EWritable, RTResource);
+	FTexture2DRHIRef RTTextureRHI = RenderTargetResource->GetTextureRHI();
+	InRHICmdList.TransitionResource(EResourceTransitionAccess::EWritable, RTTextureRHI);
 	
-	FRHIRenderPassInfo RPInfo(RTResource, ERenderTargetActions::Load_Store);
-	if (bInClearTarget)
+	FRHIRenderPassInfo RPInfo(RTTextureRHI, ERenderTargetActions::Load_Store);
+	if (Context.bClearTarget)
 	{
 		RPInfo.ColorRenderTargets[0].Action = ERenderTargetActions::Clear_Store;
 	}
@@ -154,15 +154,15 @@ void FSlate3DRenderer::DrawWindowToTarget_RenderThread( FRHICommandListImmediate
 			{
 				RenderTargetPolicy->UpdateVertexAndIndexBuffers(InRHICmdList, BatchData);
 
-				FVector2D DrawOffset = WindowDrawBuffer.ViewOffset;
+				FVector2D DrawOffset = Context.WindowDrawBuffer->ViewOffset;
 
-				FMatrix ProjectionMatrix = FSlateRHIRenderer::CreateProjectionMatrix(RTResource->GetSizeX(), RTResource->GetSizeY());
+				FMatrix ProjectionMatrix = FSlateRHIRenderer::CreateProjectionMatrix(RTTextureRHI->GetSizeX(), RTTextureRHI->GetSizeY());
 				FMatrix ViewOffset = FTranslationMatrix::Make(FVector(DrawOffset, 0));
 				ProjectionMatrix = ViewOffset * ProjectionMatrix;
 
-				FSlateBackBuffer BackBufferTarget(RenderTargetResource->GetTextureRHI(), FIntPoint(RTResource->GetSizeX(), RTResource->GetSizeY()));
+				FSlateBackBuffer BackBufferTarget(RenderTargetResource->GetTextureRHI(), FIntPoint(RTTextureRHI->GetSizeX(), RTTextureRHI->GetSizeY()));
 
-				FSlateRenderingOptions DrawOptions(ProjectionMatrix);
+				FSlateRenderingParams DrawOptions(ProjectionMatrix, Context.WorldTimeSeconds, Context.DeltaTimeSeconds, Context.RealTimeSeconds);
 				// The scene renderer will handle it in this case
 				DrawOptions.bAllowSwitchVerticalAxis = false;
 				DrawOptions.ViewOffset = DrawOffset;
@@ -195,8 +195,8 @@ void FSlate3DRenderer::DrawWindowToTarget_RenderThread( FRHICommandListImmediate
 	}
 	InRHICmdList.EndRenderPass();
 
-	FSlateEndDrawingWindowsCommand::EndDrawingWindows(InRHICmdList, &WindowDrawBuffer, *RenderTargetPolicy);
-	InRHICmdList.CopyToResolveTarget(RenderTargetResource->GetTextureRHI(), RTResource, FResolveParams());
+	FSlateEndDrawingWindowsCommand::EndDrawingWindows(InRHICmdList, Context.WindowDrawBuffer, *RenderTargetPolicy);
+	InRHICmdList.CopyToResolveTarget(RenderTargetResource->GetTextureRHI(), RTTextureRHI, FResolveParams());
 
 	ISlate3DRendererPtr Self = SharedThis(this);
 
diff --git a/Engine/Source/Runtime/SlateRHIRenderer/Private/Slate3DRenderer.h b/Engine/Source/Runtime/SlateRHIRenderer/Private/Slate3DRenderer.h
--- a/Engine/Source/Runtime/SlateRHIRenderer/Private/Slate3DRenderer.h
+++ b/Engine/Source/Runtime/SlateRHIRenderer/Private/Slate3DRenderer.h
@@ -23,7 +23,7 @@ public:
 	virtual void SetUseGammaCorrection(bool bUseGammaCorrection) override;
 	virtual FSlateDrawBuffer& GetDrawBuffer() override;
 	virtual void DrawWindow_GameThread(FSlateDrawBuffer& DrawBuffer) override;
-	virtual void DrawWindowToTarget_RenderThread(FRHICommandListImmediate& RHICmdList, FTextureRenderTarget2DResource* RenderTargetResource, FSlateDrawBuffer& InDrawBuffer, bool bInClearTarget) override;
+	virtual void DrawWindowToTarget_RenderThread(FRHICommandListImmediate& RHICmdList, const struct FRenderThreadUpdateContext& Context) override;
 
 private:
 
diff --git a/Engine/Source/Runtime/SlateRHIRenderer/Private/SlateRHIRenderer.cpp b/Engine/Source/Runtime/SlateRHIRenderer/Private/SlateRHIRenderer.cpp
--- a/Engine/Source/Runtime/SlateRHIRenderer/Private/SlateRHIRenderer.cpp
+++ b/Engine/Source/Runtime/SlateRHIRenderer/Private/SlateRHIRenderer.cpp
@@ -77,20 +77,31 @@ TAutoConsoleVariable<int32> CVarShowSlateBatching(
 );
 #endif
 
-void FSlateRHIRenderer::FViewportInfo::InitRHI()
+struct FSlateDrawWindowCommandParams
+{
+	FSlateRHIRenderer* Renderer;
+	FSlateWindowElementList* WindowElementList;
+	float WorldTimeSeconds;
+	float DeltaTimeSeconds;
+	float RealTimeSeconds;
+	bool bLockToVsync;
+	bool bClear;
+};
+
+void FViewportInfo::InitRHI()
 {
 	// Viewport RHI is created on the game thread
 	// Create the depth-stencil surface if needed.
 	RecreateDepthBuffer_RenderThread();
 }
 
-void FSlateRHIRenderer::FViewportInfo::ReleaseRHI()
+void FViewportInfo::ReleaseRHI()
 {
 	DepthStencil.SafeRelease();
 	ViewportRHI.SafeRelease();
 }
 
-void FSlateRHIRenderer::FViewportInfo::ConditionallyUpdateDepthBuffer(bool bInRequiresStencilTest, uint32 InWidth, uint32 InHeight)
+void FViewportInfo::ConditionallyUpdateDepthBuffer(bool bInRequiresStencilTest, uint32 InWidth, uint32 InHeight)
 {
 	FViewportInfo* ViewportInfo = this;
 	ENQUEUE_RENDER_COMMAND(UpdateDepthBufferCommand)(
@@ -113,7 +124,7 @@ void FSlateRHIRenderer::FViewportInfo::ConditionallyUpdateDepthBuffer(bool bInRe
 	);
 }
 
-void FSlateRHIRenderer::FViewportInfo::RecreateDepthBuffer_RenderThread()
+void FViewportInfo::RecreateDepthBuffer_RenderThread()
 {
 	check(IsInRenderingThread());
 	DepthStencil.SafeRelease();
@@ -629,7 +640,7 @@ int32 SlateWireFrame = 0;
 static FAutoConsoleVariableRef CVarSlateWireframe(TEXT("Slate.ShowWireFrame"), SlateWireFrame, TEXT(""), ECVF_Default);
 
 /** Draws windows from a FSlateDrawBuffer on the render thread */
-void FSlateRHIRenderer::DrawWindow_RenderThread(FRHICommandListImmediate& RHICmdList, FViewportInfo& ViewportInfo, FSlateWindowElementList& WindowElementList, bool bLockToVsync, bool bClear)
+void FSlateRHIRenderer::DrawWindow_RenderThread(FRHICommandListImmediate& RHICmdList, FViewportInfo& ViewportInfo, FSlateWindowElementList& WindowElementList, const struct FSlateDrawWindowCommandParams& DrawCommandParams)
 {
 	static uint32 LastTimestamp = FPlatformTime::Cycles();
 	{
@@ -729,6 +740,7 @@ void FSlateRHIRenderer::DrawWindow_RenderThread(FRHICommandListImmediate& RHICmd
 
 			FTexture2DRHIRef FinalBuffer = BackBuffer;
 
+			bool bClear = DrawCommandParams.bClear;
 			if (bCompositeUI)
 			{
 				FResolveParams ResolveParams;
@@ -788,8 +800,8 @@ void FSlateRHIRenderer::DrawWindow_RenderThread(FRHICommandListImmediate& RHICmd
 
 						FSlateBackBuffer BackBufferTarget(BackBuffer, FIntPoint(ViewportWidth, ViewportHeight));
 
-						FSlateRenderingOptions DrawOptions(ViewMatrix * ViewportInfo.ProjectionMatrix);
-						DrawOptions.bWireFrame = !!SlateWireFrame;
+						FSlateRenderingParams RenderParams(ViewMatrix * ViewportInfo.ProjectionMatrix, DrawCommandParams.WorldTimeSeconds, DrawCommandParams.DeltaTimeSeconds, DrawCommandParams.RealTimeSeconds);
+						RenderParams.bWireFrame = !!SlateWireFrame;
 
 						RenderingPolicy->DrawElements
 						(
@@ -798,7 +810,7 @@ void FSlateRHIRenderer::DrawWindow_RenderThread(FRHICommandListImmediate& RHICmd
 							BackBuffer,
 							ViewportInfo.DepthStencil,
 							BatchData.GetRenderBatches(),
-							DrawOptions
+							RenderParams
 						);
 					}
 				}
@@ -940,7 +952,7 @@ void FSlateRHIRenderer::DrawWindow_RenderThread(FRHICommandListImmediate& RHICmd
 	// Calculate renderthread time (excluding idle time).	
 	uint32 StartTime = FPlatformTime::Cycles();
 
-	RHICmdList.EndDrawingViewport(ViewportInfo.ViewportRHI, true, bLockToVsync);
+	RHICmdList.EndDrawingViewport(ViewportInfo.ViewportRHI, true, DrawCommandParams.bLockToVsync);
 
 	uint32 EndTime = FPlatformTime::Cycles();
 
@@ -1090,34 +1102,29 @@ void FSlateRHIRenderer::DrawWindows_Private(FSlateDrawBuffer& WindowDrawBuffer)
 
 				// Tell the rendering thread to draw the windows
 				{
-					struct FSlateDrawWindowCommandParams
-					{
-						FSlateRHIRenderer* Renderer;
-						FSlateRHIRenderer::FViewportInfo* ViewportInfo;
-						FSlateWindowElementList* WindowElementList;
-						SWindow* SlateWindow;
-						bool bLockToVsync;
-						bool bClear;
-					} Params;
+					FSlateDrawWindowCommandParams Params;
 
 					Params.Renderer = this;
-					Params.ViewportInfo = ViewInfo;
 					Params.WindowElementList = &ElementList;
 					Params.bLockToVsync = bLockToVsync;
 #if ALPHA_BLENDED_WINDOWS
 					Params.bClear = Window->GetTransparencySupport() == EWindowTransparency::PerPixel;
 #else
 					Params.bClear = false;
-#endif
+#endif	
+					Params.WorldTimeSeconds = FApp::GetCurrentTime() - GStartTime;
+					Params.DeltaTimeSeconds = FApp::GetDeltaTime();
+					Params.RealTimeSeconds = FPlatformTime::Seconds() - GStartTime;
+
 					// Skip the actual draw if we're in a headless execution environment
 					bool bLocalTakingAScreenShot = bTakingAScreenShot;
 					if (GIsClient && !IsRunningCommandlet() && !GUsingNullRHI)
 					{
 						ENQUEUE_RENDER_COMMAND(SlateDrawWindowsCommand)(
-							[Params](FRHICommandListImmediate& RHICmdList)
-						{
-							Params.Renderer->DrawWindow_RenderThread(RHICmdList, *Params.ViewportInfo, *Params.WindowElementList, Params.bLockToVsync, Params.bClear);
-						}
+							[Params, ViewInfo](FRHICommandListImmediate& RHICmdList)
+							{
+								Params.Renderer->DrawWindow_RenderThread(RHICmdList, *ViewInfo, *Params.WindowElementList, Params);
+							}
 						);
 					}
 
@@ -1151,10 +1158,9 @@ void FSlateRHIRenderer::DrawWindows_Private(FSlateDrawBuffer& WindowDrawBuffer)
 		ENQUEUE_UNIQUE_RENDER_COMMAND_ONEPARAMETER(DrawWidgetRendererImmediate,
 			FDeferredUpdateContextList, Contexts, DeferredUpdateContexts,
 			{
-
 				for (const FRenderThreadUpdateContext& Context : Contexts)
 				{
-					static_cast<ISlate3DRenderer*>(Context.Renderer)->DrawWindowToTarget_RenderThread(RHICmdList, static_cast<FTextureRenderTarget2DResource*>(Context.RenderTargetResource), *Context.DrawBuffer, Context.bClearTarget);
+					static_cast<ISlate3DRenderer*>(Context.Renderer)->DrawWindowToTarget_RenderThread(RHICmdList, Context);
 				}
 			}
 		);
@@ -1516,7 +1522,7 @@ void FSlateRHIRenderer::AddWidgetRendererUpdate(const struct FRenderThreadUpdate
 		ENQUEUE_UNIQUE_RENDER_COMMAND_ONEPARAMETER(DrawWidgetRendererImmediate,
 			FRenderThreadUpdateContext, InContext, Context,
 			{
-				static_cast<ISlate3DRenderer*>(InContext.Renderer)->DrawWindowToTarget_RenderThread(RHICmdList, static_cast<FTextureRenderTarget2DResource*>(InContext.RenderTargetResource), *InContext.DrawBuffer, InContext.bClearTarget);
+				static_cast<ISlate3DRenderer*>(InContext.Renderer)->DrawWindowToTarget_RenderThread(RHICmdList, InContext);
 			});
 	}
 }
diff --git a/Engine/Source/Runtime/SlateRHIRenderer/Private/SlateRHIRenderer.h b/Engine/Source/Runtime/SlateRHIRenderer/Private/SlateRHIRenderer.h
--- a/Engine/Source/Runtime/SlateRHIRenderer/Private/SlateRHIRenderer.h
+++ b/Engine/Source/Runtime/SlateRHIRenderer/Private/SlateRHIRenderer.h
@@ -47,115 +47,112 @@ private:
 	FIntPoint SizeXY;
 };
 
-/** A Slate rendering implementation for Unreal engine */
-class FSlateRHIRenderer : public FSlateRenderer
-{
-private:
 
-	/** An RHI Representation of a viewport with cached width and height for detecting resizes */
-	struct FViewportInfo : public FRenderResource
+/** An RHI Representation of a viewport with cached width and height for detecting resizes */
+struct FViewportInfo : public FRenderResource
+{
+	/** The projection matrix used in the viewport */
+	FMatrix ProjectionMatrix;
+	/** The viewport rendering handle */
+	FViewportRHIRef ViewportRHI;
+	/** The depth buffer texture if any */
+	FTexture2DRHIRef DepthStencil;
+
+	// Buffers used in HDR compositing
+	/** sRGB UI render target */
+	FTexture2DRHIRef UITargetRT;
+	/** HDR source data */
+	FTexture2DRHIRef HDRSourceRT;
+	/** sRGB UI render target */
+	FTexture2DRHIRef UITargetSRV;
+	/** HDR source data */
+	FTexture2DRHIRef HDRSourceSRV;
+
+	/** Color-space LUT for HDR UI composition. */
+	FTexture3DRHIRef ColorSpaceLUTRT;
+	FTexture3DRHIRef ColorSpaceLUTSRV;
+	int32 ColorSpaceLUTOutputDevice;
+	int32 ColorSpaceLUTOutputGamut;
+
+	//FTexture2DRHIRef RenderTargetTexture;
+	/** The OS Window handle (for recreating the viewport) */
+	void* OSWindow;
+	/** The actual width of the viewport */
+	uint32 Width;
+	/** The actual height of the viewport */
+	uint32 Height;
+	/** The desired width of the viewport */
+	uint32 DesiredWidth;
+	/** The desired height of the viewport */
+	uint32 DesiredHeight;
+	/** Whether or not the viewport requires a stencil test */
+	bool bRequiresStencilTest;
+	/** Whether or not the viewport is in fullscreen */
+	bool bFullscreen;
+	/** The desired pixel format for this viewport */
+	EPixelFormat PixelFormat;
+	/** The desired SDR pixel format for this viewport */
+	EPixelFormat SDRPixelFormat;
+	/** Color gamut for output to HDR display */
+	int32 HDRColorGamut;
+	/** Device format for output to HDR display */
+	int32 HDROutputDevice;
+
+	IViewportRenderTargetProvider* RTProvider;
+
+	/** FRenderResource interface */
+	virtual void InitRHI() override;
+	virtual void ReleaseRHI() override;
+
+	FViewportInfo()
+		: ColorSpaceLUTOutputDevice(0),
+		ColorSpaceLUTOutputGamut(0),
+		OSWindow(NULL),
+		Width(0),
+		Height(0),
+		DesiredWidth(0),
+		DesiredHeight(0),
+		bRequiresStencilTest(false),
+		bFullscreen(false),
+		PixelFormat(EPixelFormat::PF_Unknown),
+		SDRPixelFormat(EPixelFormat::PF_Unknown),
+		RTProvider(nullptr)
 	{
-		/** The projection matrix used in the viewport */
-		FMatrix ProjectionMatrix;	
-		/** The viewport rendering handle */
-		FViewportRHIRef ViewportRHI;
-		/** The depth buffer texture if any */
-		FTexture2DRHIRef DepthStencil;
-
-		// Buffers used in HDR compositing
-		/** sRGB UI render target */
-		FTexture2DRHIRef UITargetRT;
-		/** HDR source data */
-		FTexture2DRHIRef HDRSourceRT;
-		/** sRGB UI render target */
-		FTexture2DRHIRef UITargetSRV;
-		/** HDR source data */
-		FTexture2DRHIRef HDRSourceSRV;
-
-		/** Color-space LUT for HDR UI composition. */
-		FTexture3DRHIRef ColorSpaceLUTRT;
-		FTexture3DRHIRef ColorSpaceLUTSRV;
-		int32 ColorSpaceLUTOutputDevice;
-		int32 ColorSpaceLUTOutputGamut;
-		
-		//FTexture2DRHIRef RenderTargetTexture;
-		/** The OS Window handle (for recreating the viewport) */
-		void* OSWindow;
-		/** The actual width of the viewport */
-		uint32 Width;
-		/** The actual height of the viewport */
-		uint32 Height;
-		/** The desired width of the viewport */
-		uint32 DesiredWidth;
-		/** The desired height of the viewport */
-		uint32 DesiredHeight;
-		/** Whether or not the viewport requires a stencil test */
-		bool bRequiresStencilTest;
-		/** Whether or not the viewport is in fullscreen */
-		bool bFullscreen;
-		/** The desired pixel format for this viewport */
-		EPixelFormat PixelFormat;
-		/** The desired SDR pixel format for this viewport */
-		EPixelFormat SDRPixelFormat;
-		/** Color gamut for output to HDR display */
-		int32 HDRColorGamut;
-		/** Device format for output to HDR display */
-		int32 HDROutputDevice;
-
-		IViewportRenderTargetProvider* RTProvider;
-	
-		/** FRenderResource interface */
-		virtual void InitRHI() override;
-		virtual void ReleaseRHI() override;
-
-		FViewportInfo()
-			:	ColorSpaceLUTOutputDevice(0),
-				ColorSpaceLUTOutputGamut(0),
-				OSWindow(NULL), 
-				Width(0),
-				Height(0),
-				DesiredWidth(0),
-				DesiredHeight(0),
-				bRequiresStencilTest(false),
-				bFullscreen(false),
-				PixelFormat(EPixelFormat::PF_Unknown),
-				SDRPixelFormat(EPixelFormat::PF_Unknown),
-				RTProvider(nullptr)
-		{
 
-		}
+	}
 
-		~FViewportInfo()
-		{
-			DepthStencil.SafeRelease();
-			UITargetRT.SafeRelease();
-			HDRSourceRT.SafeRelease();
-			UITargetSRV.SafeRelease();
-			HDRSourceSRV.SafeRelease();
-			ColorSpaceLUTRT.SafeRelease();
-			ColorSpaceLUTSRV.SafeRelease();
-		}
+	~FViewportInfo()
+	{
+		DepthStencil.SafeRelease();
+		UITargetRT.SafeRelease();
+		HDRSourceRT.SafeRelease();
+		UITargetSRV.SafeRelease();
+		HDRSourceSRV.SafeRelease();
+		ColorSpaceLUTRT.SafeRelease();
+		ColorSpaceLUTSRV.SafeRelease();
+	}
 
-		void ConditionallyUpdateDepthBuffer(bool bInRequiresStencilTest, uint32 Width, uint32 Height);
-		void RecreateDepthBuffer_RenderThread();
+	void ConditionallyUpdateDepthBuffer(bool bInRequiresStencilTest, uint32 Width, uint32 Height);
+	void RecreateDepthBuffer_RenderThread();
 
-		FTexture2DRHIRef GetRenderTargetTexture() const
+	FTexture2DRHIRef GetRenderTargetTexture() const
+	{
+		if (RTProvider)
 		{
-			if (RTProvider)
+			FSlateShaderResource* RenderTargetTexture = RTProvider->GetViewportRenderTargetTexture();
+			if (RenderTargetTexture)
 			{
-				FSlateShaderResource* RenderTargetTexture = RTProvider->GetViewportRenderTargetTexture();
-				if( RenderTargetTexture )
-				{
-					FSlateRenderTargetRHI* RHITarget = (FSlateRenderTargetRHI*)RenderTargetTexture;
-					return RHITarget->GetTypedResource();
-				}
+				FSlateRenderTargetRHI* RHITarget = (FSlateRenderTargetRHI*)RenderTargetTexture;
+				return RHITarget->GetTypedResource();
 			}
-			return nullptr;
 		}
+		return nullptr;
+	}
+};
 
-	private:		
-	};
-
+/** A Slate rendering implementation for Unreal engine */
+class FSlateRHIRenderer : public FSlateRenderer
+{
 public:
 	FSlateRHIRenderer( TSharedRef<FSlateFontServices> InSlateFontServices, TSharedRef<FSlateRHIResourceManager> InResourceManager );
 	~FSlateRHIRenderer();
@@ -206,7 +203,7 @@ public:
 	virtual void AddWidgetRendererUpdate(const struct FRenderThreadUpdateContext& Context, bool bDeferredRenderTargetUpdate) override;
 
 	/** Draws windows from a FSlateDrawBuffer on the render thread */
-	void DrawWindow_RenderThread(FRHICommandListImmediate& RHICmdList, FSlateRHIRenderer::FViewportInfo& ViewportInfo, FSlateWindowElementList& WindowElementList, bool bLockToVsync, bool bClear);
+	void DrawWindow_RenderThread(FRHICommandListImmediate& RHICmdList, FViewportInfo& ViewportInfo, FSlateWindowElementList& WindowElementList, const struct FSlateDrawWindowCommandParams& DrawParams);
 
 	/**
 	 * Reloads texture resources from disk                   
diff --git a/Engine/Source/Runtime/SlateRHIRenderer/Private/SlateRHIRenderingPolicy.cpp b/Engine/Source/Runtime/SlateRHIRenderer/Private/SlateRHIRenderingPolicy.cpp
--- a/Engine/Source/Runtime/SlateRHIRenderer/Private/SlateRHIRenderingPolicy.cpp
+++ b/Engine/Source/Runtime/SlateRHIRenderer/Private/SlateRHIRenderingPolicy.cpp
@@ -618,7 +618,7 @@ void FSlateRHIRenderingPolicy::DrawElements(
 	FTexture2DRHIRef& ColorTarget,
 	FTexture2DRHIRef& DepthStencilTarget,
 	const TArray<FSlateRenderBatch>& RenderBatches,
-	const FSlateRenderingOptions& Options)
+	const FSlateRenderingParams& Params)
 {
 	// Should only be called by the rendering thread
 	check(IsInRenderingThread());
@@ -639,10 +639,6 @@ void FSlateRHIRenderingPolicy::DrawElements(
 
 	IRendererModule& RendererModule = FModuleManager::GetModuleChecked<IRendererModule>(RendererModuleName);
 
-	float TimeSeconds = FApp::GetCurrentTime() - GStartTime;
-	float DeltaTimeSeconds = FApp::GetDeltaTime();
-	float RealTimeSeconds = FPlatformTime::Seconds() - GStartTime;
-
 	static const FEngineShowFlags DefaultShowFlags(ESFIM_Game);
 
 	const float EngineGamma = GEngine ? GEngine->GetDisplayGamma() : 2.2f;
@@ -665,7 +661,7 @@ void FSlateRHIRenderingPolicy::DrawElements(
 		QUICK_SCOPE_CYCLE_COUNTER(STAT_Slate_CreateScenes);
 		for (int32 i = 0; i < ResourceManager->GetSceneCount(); i++)
 		{
-			SceneViewFamilyContexts[i] = new  FSceneViewFamilyContext
+			SceneViewFamilyContexts[i] = new FSceneViewFamilyContext
 			(
 				FSceneViewFamily::ConstructionValues
 				(
@@ -673,11 +669,11 @@ void FSlateRHIRenderingPolicy::DrawElements(
 					ResourceManager->GetSceneAt(i),
 					DefaultShowFlags
 				)
-				.SetWorldTimes(TimeSeconds, DeltaTimeSeconds, RealTimeSeconds)
+				.SetWorldTimes(Params.CurrentWorldTime, Params.DeltaTimeSeconds, Params.CurrentRealTime)
 				.SetGammaCorrection(DisplayGamma)
 				.SetRealtimeUpdate(true)
 			);
-			SceneViews[i] = CreateSceneView(SceneViewFamilyContexts[i], BackBuffer, Options.ViewProjectionMatrix);
+			SceneViews[i] = CreateSceneView(SceneViewFamilyContexts[i], BackBuffer, Params.ViewProjectionMatrix);
 		}
 
 		SceneViewFamilyContexts[NumScenes - 1] = new FSceneViewFamilyContext
@@ -688,11 +684,11 @@ void FSlateRHIRenderingPolicy::DrawElements(
 				nullptr,
 				DefaultShowFlags
 			)
-			.SetWorldTimes(TimeSeconds, DeltaTimeSeconds, RealTimeSeconds)
+			.SetWorldTimes(Params.CurrentWorldTime, Params.DeltaTimeSeconds, Params.CurrentRealTime)
 			.SetGammaCorrection(DisplayGamma)
 			.SetRealtimeUpdate(true)
 		);
-		SceneViews[NumScenes - 1] = CreateSceneView(SceneViewFamilyContexts[NumScenes - 1], BackBuffer, Options.ViewProjectionMatrix);
+		SceneViews[NumScenes - 1] = CreateSceneView(SceneViewFamilyContexts[NumScenes - 1], BackBuffer, Params.ViewProjectionMatrix);
 	}
 
 	TShaderMapRef<FSlateElementVS> GlobalVertexShader(GetGlobalShaderMap(GMaxRHIFeatureLevel));
@@ -715,7 +711,7 @@ void FSlateRHIRenderingPolicy::DrawElements(
 #endif
 
 	const bool bAbsoluteIndices = CVarSlateAbsoluteIndices.GetValueOnRenderThread() != 0;
-	const bool bSwitchVerticalAxis = Options.bAllowSwitchVerticalAxis && RHINeedsToSwitchVerticalAxis(GShaderPlatformForFeatureLevel[GMaxRHIFeatureLevel]);
+	const bool bSwitchVerticalAxis = Params.bAllowSwitchVerticalAxis && RHINeedsToSwitchVerticalAxis(GShaderPlatformForFeatureLevel[GMaxRHIFeatureLevel]);
 
 	// This variable tracks the last clipping state, so that if multiple batches have the same clipping state, we don't have to do any work.
 	TOptional<FSlateClippingState> LastClippingState;
@@ -729,7 +725,7 @@ void FSlateRHIRenderingPolicy::DrawElements(
 	// Disable depth/stencil testing by default
 	GraphicsPSOInit.DepthStencilState = TStaticDepthStencilState<false, CF_Always>::GetRHI();
 
-	FVector2D ViewTranslation2D = Options.ViewOffset;
+	FVector2D ViewTranslation2D = Params.ViewOffset;
 
 	// Draw each element
 	/*
@@ -799,7 +795,7 @@ void FSlateRHIRenderingPolicy::DrawElements(
 		if (!RenderBatch.CustomDrawer.IsValid())
 		{
 			FMatrix DynamicOffset = FTranslationMatrix::Make(FVector(RenderBatch.DynamicOffset.X, RenderBatch.DynamicOffset.Y, 0));
-			const FMatrix ViewProjection = DynamicOffset * Options.ViewProjectionMatrix;
+			const FMatrix ViewProjection = DynamicOffset * Params.ViewProjectionMatrix;
 
 			UpdateScissorRect(
 				RHICmdList,
@@ -866,11 +862,11 @@ void FSlateRHIRenderingPolicy::DrawElements(
 						;
 				}
 
-				if (EnumHasAllFlags(DrawFlags, ESlateBatchDrawFlag::Wireframe) || Options.bWireFrame)
+				if (EnumHasAllFlags(DrawFlags, ESlateBatchDrawFlag::Wireframe) || Params.bWireFrame)
 				{
 					GraphicsPSOInit.RasterizerState = TStaticRasterizerState<FM_Wireframe, CM_None, false>::GetRHI();
 
-					if (Options.bWireFrame)
+					if (Params.bWireFrame)
 					{
 						GraphicsPSOInit.BlendState = TStaticBlendState<>::GetRHI();
 					}
@@ -1227,7 +1223,7 @@ void FSlateRHIRenderingPolicy::DrawElements(
 						ViewTranslation2D,
 						bSwitchVerticalAxis,
 						InGraphicsPSOInit,
-						Options.ViewProjectionMatrix,
+						Params.ViewProjectionMatrix,
 						true);
 				};
 
diff --git a/Engine/Source/Runtime/SlateRHIRenderer/Private/SlateRHIRenderingPolicy.h b/Engine/Source/Runtime/SlateRHIRenderer/Private/SlateRHIRenderingPolicy.h
--- a/Engine/Source/Runtime/SlateRHIRenderer/Private/SlateRHIRenderingPolicy.h
+++ b/Engine/Source/Runtime/SlateRHIRenderer/Private/SlateRHIRenderingPolicy.h
@@ -21,16 +21,22 @@ class FSlatePostProcessor;
 class ILayoutCache;
 class UDeviceProfile;
 
-struct FSlateRenderingOptions
+struct FSlateRenderingParams
 {
 	FMatrix ViewProjectionMatrix;
 	FVector2D ViewOffset;
+	float CurrentWorldTime;
+	float DeltaTimeSeconds;
+	float CurrentRealTime;
 	bool bAllowSwitchVerticalAxis;
 	bool bWireFrame;
 
-	FSlateRenderingOptions(const FMatrix& InViewProjectionMatrix)
+	FSlateRenderingParams(const FMatrix& InViewProjectionMatrix, float InCurrentWorldTime, float InDeltaTimeSeconds, float InCurrentRealTime)
 		: ViewProjectionMatrix(InViewProjectionMatrix)
 		, ViewOffset(0, 0)
+		, CurrentWorldTime(InCurrentWorldTime)
+		, DeltaTimeSeconds(InDeltaTimeSeconds)
+		, CurrentRealTime(InCurrentRealTime)
 		, bAllowSwitchVerticalAxis(true)
 		, bWireFrame(false)
 	{
@@ -47,7 +53,7 @@ public:
 
 	void ReleaseCachingResourcesFor(FRHICommandListImmediate& RHICmdList, const ILayoutCache* Cacher);
 
-	void DrawElements(FRHICommandListImmediate& RHICmdList, class FSlateBackBuffer& BackBuffer, FTexture2DRHIRef& ColorTarget, FTexture2DRHIRef& DepthStencilTarget, const TArray<FSlateRenderBatch>& RenderBatches, const FSlateRenderingOptions& Options);
+	void DrawElements(FRHICommandListImmediate& RHICmdList, class FSlateBackBuffer& BackBuffer, FTexture2DRHIRef& ColorTarget, FTexture2DRHIRef& DepthStencilTarget, const TArray<FSlateRenderBatch>& RenderBatches, const FSlateRenderingParams& Params);
 
 	virtual TSharedRef<FSlateShaderResourceManager> GetResourceManager() const override { return ResourceManager; }
 	virtual bool IsVertexColorInLinearSpace() const override { return false; }
diff --git a/Engine/Source/Runtime/SlateRHIRenderer/Public/Interfaces/ISlate3DRenderer.h b/Engine/Source/Runtime/SlateRHIRenderer/Public/Interfaces/ISlate3DRenderer.h
--- a/Engine/Source/Runtime/SlateRHIRenderer/Public/Interfaces/ISlate3DRenderer.h
+++ b/Engine/Source/Runtime/SlateRHIRenderer/Public/Interfaces/ISlate3DRenderer.h
@@ -34,7 +34,7 @@ public:
 	 * @param RenderTarget The render target to render the contents of the draw buffer to
 	 * @param InDrawBuffer The draw buffer containing the batched elements to render
 	 */
-	virtual void DrawWindowToTarget_RenderThread(FRHICommandListImmediate& RHICmdList, FTextureRenderTarget2DResource* RenderTargetResource, FSlateDrawBuffer& InDrawBuffer, bool bInClearTarget = true) = 0;
+	virtual void DrawWindowToTarget_RenderThread(FRHICommandListImmediate& RHICmdList, const struct FRenderThreadUpdateContext& Context) = 0;
 };
 
 typedef TSharedPtr<ISlate3DRenderer, ESPMode::ThreadSafe> ISlate3DRendererPtr;
diff --git a/Engine/Source/Runtime/UMG/Private/Slate/WidgetRenderer.cpp b/Engine/Source/Runtime/UMG/Private/Slate/WidgetRenderer.cpp
--- a/Engine/Source/Runtime/UMG/Private/Slate/WidgetRenderer.cpp
+++ b/Engine/Source/Runtime/UMG/Private/Slate/WidgetRenderer.cpp
@@ -186,6 +186,9 @@ void FWidgetRenderer::DrawWindow(
 		FRenderThreadUpdateContext Context =
 		{
 			&DrawBuffer,
+			FApp::GetCurrentTime() - GStartTime,
+			FApp::GetDeltaTime(),
+			FPlatformTime::Seconds() - GStartTime,
 			RenderTarget->GameThread_GetRenderTargetResource(),
 			Renderer.Get(),
 			bClearTarget
