Slate Renderpass updates.

diff --git a/Engine/Source/Runtime/Engine/Private/Slate/DebugCanvas.cpp b/Engine/Source/Runtime/Engine/Private/Slate/DebugCanvas.cpp
--- a/Engine/Source/Runtime/Engine/Private/Slate/DebugCanvas.cpp
+++ b/Engine/Source/Runtime/Engine/Private/Slate/DebugCanvas.cpp
@@ -206,6 +206,7 @@ void FDebugCanvasDrawer::InitDebugCanvas(FViewportClient* ViewportClient, UWorld
 void FDebugCanvasDrawer::DrawRenderThread(FRHICommandListImmediate& RHICmdList, const void* InWindowBackBuffer)
 {
 	check( IsInRenderingThread() );
+	check(RHICmdList.IsOutsideRenderPass());
 
 	SCOPED_DRAW_EVENT(RHICmdList, DrawDebugCanvas);
 
diff --git a/Engine/Source/Runtime/Engine/Private/Slate/SceneViewport.cpp b/Engine/Source/Runtime/Engine/Private/Slate/SceneViewport.cpp
--- a/Engine/Source/Runtime/Engine/Private/Slate/SceneViewport.cpp
+++ b/Engine/Source/Runtime/Engine/Private/Slate/SceneViewport.cpp
@@ -1687,7 +1687,6 @@ void FSceneViewport::BeginRenderFrame(FRHICommandListImmediate& RHICmdList)
 	if (UseSeparateRenderTarget())
 	{		
 		RHICmdList.TransitionResource(EResourceTransitionAccess::EWritable, RenderTargetTextureRenderThreadRHI);
-		SetRenderTarget(RHICmdList,  RenderTargetTextureRenderThreadRHI,  FTexture2DRHIRef(), true);
 	}
 	else if( IsValidRef( ViewportRHI ) ) 
 	{
@@ -1714,7 +1713,7 @@ void FSceneViewport::EndRenderFrame(FRHICommandListImmediate& RHICmdList, bool b
 		if (bShouldUnsetTargets)
 		{
 			// Set the active render target(s) to nothing to release references in the case that the viewport is resized by slate before we draw again
-			SetRenderTarget(RHICmdList,  FTexture2DRHIRef(), FTexture2DRHIRef() );
+			UnbindRenderTargets(RHICmdList);
 		}
 		// Note: this releases our reference but does not release the resource as it is owned by slate (this is intended)
 		RenderTargetTextureRenderThreadRHI.SafeRelease();
diff --git a/Engine/Source/Runtime/Engine/Private/Slate/SlateTextures.cpp b/Engine/Source/Runtime/Engine/Private/Slate/SlateTextures.cpp
--- a/Engine/Source/Runtime/Engine/Private/Slate/SlateTextures.cpp
+++ b/Engine/Source/Runtime/Engine/Private/Slate/SlateTextures.cpp
@@ -360,9 +360,9 @@ void FSlateTextureRenderTarget2DResource::UpdateDeferredResource(FRHICommandList
 	// Clear the target surface to green
 	if (bClearRenderTarget)
 	{
-		FRHIRenderTargetView View = FRHIRenderTargetView(RenderTargetTextureRHI, ERenderTargetLoadAction::EClear);
-		FRHISetRenderTargetsInfo Info(1, &View, FRHIDepthRenderTargetView());
-		RHICmdList.SetRenderTargetsAndClear(Info);
+		FRHIRenderPassInfo RPInfo(RenderTargetTextureRHI, ERenderTargetActions::Clear_Store);
+		RHICmdList.BeginRenderPass(RPInfo, TEXT("Slate2DUpdateDeferred_Clear"));
+		RHICmdList.EndRenderPass();
 	}
 
 	// Copy surface to the texture for use
diff --git a/Engine/Source/Runtime/SlateRHIRenderer/Private/Slate3DRenderer.cpp b/Engine/Source/Runtime/SlateRHIRenderer/Private/Slate3DRenderer.cpp
--- a/Engine/Source/Runtime/SlateRHIRenderer/Private/Slate3DRenderer.cpp
+++ b/Engine/Source/Runtime/SlateRHIRenderer/Private/Slate3DRenderer.cpp
@@ -129,68 +129,71 @@ void FSlate3DRenderer::DrawWindowToTarget_RenderThread( FRHICommandListImmediate
 
 	// Set render target and clear.
 	FTexture2DRHIRef RTResource = RenderTargetResource->GetTextureRHI();
-	FRHIRenderTargetView ColorRTV(RTResource, ERenderTargetLoadAction::ELoad);
-	ColorRTV.LoadAction = ERenderTargetLoadAction::EClear;
-	FRHISetRenderTargetsInfo Info(1, &ColorRTV, FRHIDepthRenderTargetView());
-	Info.bClearColor = bInClearTarget;
-
 	InRHICmdList.TransitionResource(EResourceTransitionAccess::EWritable, RTResource);
-	InRHICmdList.SetRenderTargetsAndClear(Info);
-
-	for ( int32 WindowIndex = 0; WindowIndex < WindowsToDraw.Num(); WindowIndex++ )
+	
+	FRHIRenderPassInfo RPInfo(RTResource, ERenderTargetActions::Load_Store);
+	if (bInClearTarget)
+	{
+		RPInfo.ColorRenderTargets[0].Action = ERenderTargetActions::Clear_Store;
+	}
+	InRHICmdList.BeginRenderPass(RPInfo, TEXT("Slate3D"));
 	{
-		FSlateWindowElementList& WindowElementList = *WindowsToDraw[WindowIndex];
 
-		FSlateBatchData& BatchData = WindowElementList.GetBatchData();
-		FElementBatchMap& RootBatchMap = WindowElementList.GetRootDrawLayer().GetElementBatchMap();
+		for (int32 WindowIndex = 0; WindowIndex < WindowsToDraw.Num(); WindowIndex++)
+		{
+			FSlateWindowElementList& WindowElementList = *WindowsToDraw[WindowIndex];
+
+			FSlateBatchData& BatchData = WindowElementList.GetBatchData();
+			FElementBatchMap& RootBatchMap = WindowElementList.GetRootDrawLayer().GetElementBatchMap();
 
-		WindowElementList.PreDraw_ParallelThread();
+			WindowElementList.PreDraw_ParallelThread();
 
-		BatchData.CreateRenderBatches(RootBatchMap);
+			BatchData.CreateRenderBatches(RootBatchMap);
 
-		if ( BatchData.GetRenderBatches().Num() > 0 )
-		{
-			RenderTargetPolicy->UpdateVertexAndIndexBuffers(InRHICmdList, BatchData);
-		
-			FVector2D DrawOffset = WindowDrawBuffer.ViewOffset;
+			if (BatchData.GetRenderBatches().Num() > 0)
+			{
+				RenderTargetPolicy->UpdateVertexAndIndexBuffers(InRHICmdList, BatchData);
 
-			FMatrix ProjectionMatrix = FSlateRHIRenderer::CreateProjectionMatrix(RTResource->GetSizeX(), RTResource->GetSizeY());
-			FMatrix ViewOffset = FTranslationMatrix::Make(FVector(DrawOffset, 0));
-			ProjectionMatrix = ViewOffset * ProjectionMatrix;
-		
-			FSlateBackBuffer BackBufferTarget(RenderTargetResource->GetTextureRHI(), FIntPoint(RTResource->GetSizeX(), RTResource->GetSizeY()));
+				FVector2D DrawOffset = WindowDrawBuffer.ViewOffset;
 
-			FSlateRenderingOptions DrawOptions(ProjectionMatrix);
-			// The scene renderer will handle it in this case
-			DrawOptions.bAllowSwitchVerticalAxis = false;
-			DrawOptions.ViewOffset = DrawOffset;
+				FMatrix ProjectionMatrix = FSlateRHIRenderer::CreateProjectionMatrix(RTResource->GetSizeX(), RTResource->GetSizeY());
+				FMatrix ViewOffset = FTranslationMatrix::Make(FVector(DrawOffset, 0));
+				ProjectionMatrix = ViewOffset * ProjectionMatrix;
 
-			FTexture2DRHIRef ColorTarget = RenderTargetResource->GetTextureRHI();
+				FSlateBackBuffer BackBufferTarget(RenderTargetResource->GetTextureRHI(), FIntPoint(RTResource->GetSizeX(), RTResource->GetSizeY()));
 
-			if (BatchData.IsStencilClippingRequired())
-			{
-				if (!DepthStencil.IsValid() || ColorTarget->GetSizeXY() != DepthStencil->GetSizeXY())
-				{
-					DepthStencil.SafeRelease();
+				FSlateRenderingOptions DrawOptions(ProjectionMatrix);
+				// The scene renderer will handle it in this case
+				DrawOptions.bAllowSwitchVerticalAxis = false;
+				DrawOptions.ViewOffset = DrawOffset;
+
+				FTexture2DRHIRef ColorTarget = RenderTargetResource->GetTextureRHI();
 
-					FTexture2DRHIRef ShaderResourceUnused;
-					FRHIResourceCreateInfo CreateInfo(FClearValueBinding::DepthZero);
-					RHICreateTargetableShaderResource2D(ColorTarget->GetSizeX(), ColorTarget->GetSizeY(), PF_DepthStencil, 1, TexCreate_None, TexCreate_DepthStencilTargetable, false, CreateInfo, DepthStencil, ShaderResourceUnused);
-					check(IsValidRef(DepthStencil));
+				if (BatchData.IsStencilClippingRequired())
+				{
+					if (!DepthStencil.IsValid() || ColorTarget->GetSizeXY() != DepthStencil->GetSizeXY())
+					{
+						DepthStencil.SafeRelease();
+
+						FTexture2DRHIRef ShaderResourceUnused;
+						FRHIResourceCreateInfo CreateInfo(FClearValueBinding::DepthZero);
+						RHICreateTargetableShaderResource2D(ColorTarget->GetSizeX(), ColorTarget->GetSizeY(), PF_DepthStencil, 1, TexCreate_None, TexCreate_DepthStencilTargetable, false, CreateInfo, DepthStencil, ShaderResourceUnused);
+						check(IsValidRef(DepthStencil));
+					}
 				}
-			}
 
-			RenderTargetPolicy->DrawElements(
-				InRHICmdList,
-				BackBufferTarget,
-				ColorTarget,
-				DepthStencil,
-				BatchData.GetRenderBatches(),
-				DrawOptions
-			);
+				RenderTargetPolicy->DrawElements(
+					InRHICmdList,
+					BackBufferTarget,
+					ColorTarget,
+					DepthStencil,
+					BatchData.GetRenderBatches(),
+					DrawOptions
+				);
+			}
 		}
 	}
-
+	InRHICmdList.EndRenderPass();
 
 	FSlateEndDrawingWindowsCommand::EndDrawingWindows(InRHICmdList, &WindowDrawBuffer, *RenderTargetPolicy);
 	InRHICmdList.CopyToResolveTarget(RenderTargetResource->GetTextureRHI(), RTResource, FResolveParams());
diff --git a/Engine/Source/Runtime/SlateRHIRenderer/Private/SlatePostProcessor.cpp b/Engine/Source/Runtime/SlateRHIRenderer/Private/SlatePostProcessor.cpp
--- a/Engine/Source/Runtime/SlateRHIRenderer/Private/SlatePostProcessor.cpp
+++ b/Engine/Source/Runtime/SlateRHIRenderer/Private/SlatePostProcessor.cpp
@@ -28,6 +28,7 @@ FSlatePostProcessor::~FSlatePostProcessor()
 void FSlatePostProcessor::BlurRect(FRHICommandListImmediate& RHICmdList, IRendererModule& RendererModule, const FBlurRectParams& Params, const FPostProcessRectParams& RectParams)
 {
 	SCOPE_CYCLE_COUNTER(STAT_SlatePostProcessingRTTime);
+	check(RHICmdList.IsOutsideRenderPass());
 
 	TArray<FVector4> WeightsAndOffsets;
 	const int32 SampleCount = ComputeBlurWeights(Params.KernelSize, Params.Strength, WeightsAndOffsets);
@@ -99,23 +100,84 @@ void FSlatePostProcessor::BlurRect(FRHICommandListImmediate& RHICmdList, IRender
 			RHICmdList.TransitionResource(EResourceTransitionAccess::EReadable, SourceTexture);
 			RHICmdList.TransitionResource(EResourceTransitionAccess::EWritable, DestTexture);
 
+			FRHIRenderPassInfo RPInfo(DestTexture, ERenderTargetActions::Load_Store);
+			RHICmdList.BeginRenderPass(RPInfo, TEXT("SlateBlurRectPass0"));
+			{
+				RHICmdList.ApplyCachedRenderTargets(GraphicsPSOInit);
+
+				GraphicsPSOInit.BoundShaderState.VertexDeclarationRHI = VertexDecl;
+				GraphicsPSOInit.BoundShaderState.VertexShaderRHI = GETSAFERHISHADER_VERTEX(*VertexShader);
+				GraphicsPSOInit.BoundShaderState.PixelShaderRHI = GETSAFERHISHADER_PIXEL(*PixelShader);
+				GraphicsPSOInit.PrimitiveType = PT_TriangleList;
+				SetGraphicsPipelineState(RHICmdList, GraphicsPSOInit);
+
+				PixelShader->SetWeightsAndOffsets(RHICmdList, WeightsAndOffsets, SampleCount);
+				PixelShader->SetTexture(RHICmdList, SourceTexture, BilinearClamp);
+
+				if (bDownsample)
+				{
+					PixelShader->SetUVBounds(RHICmdList, FVector4(FVector2D::ZeroVector, FVector2D((float)DownsampleSize.X / DestTextureWidth, (float)DownsampleSize.Y / DestTextureHeight) - HalfTexelOffset));
+					PixelShader->SetBufferSizeAndDirection(RHICmdList, InvBufferSize, FVector2D(1, 0));
+
+					RendererModule.DrawRectangle(
+						RHICmdList,
+						0, 0,
+						DownsampleSize.X, DownsampleSize.Y,
+						0, 0,
+						DownsampleSize.X, DownsampleSize.Y,
+						FIntPoint(DestTextureWidth, DestTextureHeight),
+						FIntPoint(DestTextureWidth, DestTextureHeight),
+						*VertexShader,
+						EDRF_Default);
+				}
+				else
+				{
+					const FVector2D InvSrcTetureSize(1.f / SrcTextureWidth, 1.f / SrcTextureHeight);
+
+					const FVector2D UVStart = FVector2D(DestRect.Left, DestRect.Top) * InvSrcTetureSize;
+					const FVector2D UVEnd = FVector2D(DestRect.Right, DestRect.Bottom) * InvSrcTetureSize;
+					const FVector2D SizeUV = UVEnd - UVStart;
+
+					PixelShader->SetUVBounds(RHICmdList, FVector4(UVStart, UVEnd));
+					PixelShader->SetBufferSizeAndDirection(RHICmdList, InvSrcTetureSize, FVector2D(1, 0));
+
+					RendererModule.DrawRectangle(
+						RHICmdList,
+						0, 0,
+						RequiredSize.X, RequiredSize.Y,
+						UVStart.X, UVStart.Y,
+						SizeUV.X, SizeUV.Y,
+						FIntPoint(DestTextureWidth, DestTextureHeight),
+						FIntPoint(1, 1),
+						*VertexShader,
+						EDRF_Default);
+				}
+			}
+			RHICmdList.EndRenderPass();
+		}
+		else
+		{
+			FTexture2DRHIRef SourceTexture = IntermediateTargets->GetRenderTarget(1);
+			FTexture2DRHIRef DestTexture = IntermediateTargets->GetRenderTarget(0);
 
-			SetRenderTarget(RHICmdList, DestTexture, FTextureRHIRef());
-			RHICmdList.ApplyCachedRenderTargets(GraphicsPSOInit);
-
-			GraphicsPSOInit.BoundShaderState.VertexDeclarationRHI = VertexDecl;
-			GraphicsPSOInit.BoundShaderState.VertexShaderRHI = GETSAFERHISHADER_VERTEX(*VertexShader);
-			GraphicsPSOInit.BoundShaderState.PixelShaderRHI = GETSAFERHISHADER_PIXEL(*PixelShader);
-			GraphicsPSOInit.PrimitiveType = PT_TriangleList;
-			SetGraphicsPipelineState(RHICmdList, GraphicsPSOInit);
+			RHICmdList.TransitionResource(EResourceTransitionAccess::EReadable, SourceTexture);
+			RHICmdList.TransitionResource(EResourceTransitionAccess::EWritable, DestTexture);
 
-			PixelShader->SetWeightsAndOffsets(RHICmdList, WeightsAndOffsets, SampleCount);
-			PixelShader->SetTexture(RHICmdList, SourceTexture, BilinearClamp);
-		
-			if(bDownsample)
+			FRHIRenderPassInfo RPInfo(DestTexture, ERenderTargetActions::Load_Store);
+			RHICmdList.BeginRenderPass(RPInfo, TEXT("SlateBlurRect"));
 			{
-				PixelShader->SetUVBounds(RHICmdList, FVector4(FVector2D::ZeroVector, FVector2D((float)DownsampleSize.X/DestTextureWidth, (float)DownsampleSize.Y / DestTextureHeight) - HalfTexelOffset));
-				PixelShader->SetBufferSizeAndDirection(RHICmdList, InvBufferSize, FVector2D(1, 0));
+				RHICmdList.ApplyCachedRenderTargets(GraphicsPSOInit);
+
+				GraphicsPSOInit.BoundShaderState.VertexDeclarationRHI = VertexDecl;
+				GraphicsPSOInit.BoundShaderState.VertexShaderRHI = GETSAFERHISHADER_VERTEX(*VertexShader);
+				GraphicsPSOInit.BoundShaderState.PixelShaderRHI = GETSAFERHISHADER_PIXEL(*PixelShader);
+				GraphicsPSOInit.PrimitiveType = PT_TriangleList;
+				SetGraphicsPipelineState(RHICmdList, GraphicsPSOInit);
+
+				PixelShader->SetWeightsAndOffsets(RHICmdList, WeightsAndOffsets, SampleCount);
+				PixelShader->SetUVBounds(RHICmdList, FVector4(FVector2D::ZeroVector, FVector2D((float)DownsampleSize.X / DestTextureWidth, (float)DownsampleSize.Y / DestTextureHeight) - HalfTexelOffset));
+				PixelShader->SetTexture(RHICmdList, SourceTexture, BilinearClamp);
+				PixelShader->SetBufferSizeAndDirection(RHICmdList, InvBufferSize, FVector2D(0, 1));
 
 				RendererModule.DrawRectangle(
 					RHICmdList,
@@ -128,65 +190,8 @@ void FSlatePostProcessor::BlurRect(FRHICommandListImmediate& RHICmdList, IRender
 					*VertexShader,
 					EDRF_Default);
 			}
-			else
-			{
-				const FVector2D InvSrcTetureSize(1.f / SrcTextureWidth, 1.f / SrcTextureHeight);
-
-				const FVector2D UVStart = FVector2D(DestRect.Left, DestRect.Top) * InvSrcTetureSize;
-				const FVector2D UVEnd = FVector2D(DestRect.Right, DestRect.Bottom) * InvSrcTetureSize;
-				const FVector2D SizeUV = UVEnd - UVStart;
-
-				PixelShader->SetUVBounds(RHICmdList, FVector4(UVStart, UVEnd));
-				PixelShader->SetBufferSizeAndDirection(RHICmdList, InvSrcTetureSize, FVector2D(1, 0));
-
-				RendererModule.DrawRectangle(
-					RHICmdList,
-					0, 0,
-					RequiredSize.X, RequiredSize.Y,
-					UVStart.X, UVStart.Y,
-					SizeUV.X, SizeUV.Y,
-					FIntPoint(DestTextureWidth, DestTextureHeight),
-					FIntPoint(1, 1),
-					*VertexShader,
-					EDRF_Default);
-			}
-
-		}
-		else
-		{
-			FTexture2DRHIRef SourceTexture = IntermediateTargets->GetRenderTarget(1);
-			FTexture2DRHIRef DestTexture = IntermediateTargets->GetRenderTarget(0);
-
-			RHICmdList.TransitionResource(EResourceTransitionAccess::EReadable, SourceTexture);
-			RHICmdList.TransitionResource(EResourceTransitionAccess::EWritable, DestTexture);
-
-			SetRenderTarget(RHICmdList, DestTexture, FTextureRHIRef());
-			RHICmdList.ApplyCachedRenderTargets(GraphicsPSOInit);
-
-			GraphicsPSOInit.BoundShaderState.VertexDeclarationRHI = VertexDecl;
-			GraphicsPSOInit.BoundShaderState.VertexShaderRHI = GETSAFERHISHADER_VERTEX(*VertexShader);
-			GraphicsPSOInit.BoundShaderState.PixelShaderRHI = GETSAFERHISHADER_PIXEL(*PixelShader);
-			GraphicsPSOInit.PrimitiveType = PT_TriangleList;
-			SetGraphicsPipelineState(RHICmdList, GraphicsPSOInit);
-			
-			PixelShader->SetWeightsAndOffsets(RHICmdList, WeightsAndOffsets, SampleCount);
-			PixelShader->SetUVBounds(RHICmdList, FVector4(FVector2D::ZeroVector, FVector2D((float)DownsampleSize.X/DestTextureWidth, (float)DownsampleSize.Y / DestTextureHeight) - HalfTexelOffset));
-			PixelShader->SetTexture(RHICmdList, SourceTexture, BilinearClamp);
-			PixelShader->SetBufferSizeAndDirection(RHICmdList, InvBufferSize, FVector2D(0, 1));
-
-			RendererModule.DrawRectangle(
-				RHICmdList,
-				0, 0,
-				DownsampleSize.X, DownsampleSize.Y,
-				0, 0,
-				DownsampleSize.X, DownsampleSize.Y,
-				FIntPoint(DestTextureWidth, DestTextureHeight),
-				FIntPoint(DestTextureWidth, DestTextureHeight),
-				*VertexShader,
-				EDRF_Default);
-		}
-
-	
+			RHICmdList.EndRenderPass();
+		}	
 	}
 
 #endif
@@ -240,29 +245,33 @@ void FSlatePostProcessor::ColorDeficiency(FRHICommandListImmediate& RHICmdList,
 		RHICmdList.TransitionResource(EResourceTransitionAccess::EReadable, SourceTexture);
 		RHICmdList.TransitionResource(EResourceTransitionAccess::EWritable, DestTexture);
 
-		SetRenderTarget(RHICmdList, DestTexture, FTextureRHIRef());
-		RHICmdList.ApplyCachedRenderTargets(GraphicsPSOInit);
+		FRHIRenderPassInfo RPInfo(DestTexture, ERenderTargetActions::Load_Store);
+		RHICmdList.BeginRenderPass(RPInfo, TEXT("ColorDeficiency"));
+		{
+			RHICmdList.ApplyCachedRenderTargets(GraphicsPSOInit);
 
-		GraphicsPSOInit.BoundShaderState.VertexDeclarationRHI = VertexDecl;
-		GraphicsPSOInit.BoundShaderState.VertexShaderRHI = GETSAFERHISHADER_VERTEX(*VertexShader);
-		GraphicsPSOInit.BoundShaderState.PixelShaderRHI = GETSAFERHISHADER_PIXEL(*PixelShader);
-		GraphicsPSOInit.PrimitiveType = PT_TriangleList;
-		SetGraphicsPipelineState(RHICmdList, GraphicsPSOInit);
+			GraphicsPSOInit.BoundShaderState.VertexDeclarationRHI = VertexDecl;
+			GraphicsPSOInit.BoundShaderState.VertexShaderRHI = GETSAFERHISHADER_VERTEX(*VertexShader);
+			GraphicsPSOInit.BoundShaderState.PixelShaderRHI = GETSAFERHISHADER_PIXEL(*PixelShader);
+			GraphicsPSOInit.PrimitiveType = PT_TriangleList;
+			SetGraphicsPipelineState(RHICmdList, GraphicsPSOInit);
 
-		PixelShader->SetColorRules(RHICmdList, GSlateColorDeficiencyCorrection, GSlateColorDeficiencyType, GSlateColorDeficiencySeverity);
-		PixelShader->SetShowCorrectionWithDeficiency(RHICmdList, GSlateShowColorDeficiencyCorrectionWithDeficiency);
-		PixelShader->SetTexture(RHICmdList, SourceTexture, PointClamp);
+			PixelShader->SetColorRules(RHICmdList, GSlateColorDeficiencyCorrection, GSlateColorDeficiencyType, GSlateColorDeficiencySeverity);
+			PixelShader->SetShowCorrectionWithDeficiency(RHICmdList, GSlateShowColorDeficiencyCorrectionWithDeficiency);
+			PixelShader->SetTexture(RHICmdList, SourceTexture, PointClamp);
 
-		RendererModule.DrawRectangle(
-			RHICmdList,
-			0, 0,
-			RequiredSize.X, RequiredSize.Y,
-			0, 0,
-			1, 1,
-			FIntPoint(DestTextureWidth, DestTextureHeight),
-			FIntPoint(1, 1),
-			*VertexShader,
-			EDRF_Default);
+			RendererModule.DrawRectangle(
+				RHICmdList,
+				0, 0,
+				RequiredSize.X, RequiredSize.Y,
+				0, 0,
+				1, 1,
+				FIntPoint(DestTextureWidth, DestTextureHeight),
+				FIntPoint(1, 1),
+				*VertexShader,
+				EDRF_Default);
+		}
+		RHICmdList.EndRenderPass();
 	}
 
 	const FIntPoint DownsampleSize = RequiredSize;
@@ -319,34 +328,37 @@ void FSlatePostProcessor::DownsampleRect(FRHICommandListImmediate& RHICmdList, I
 		RHICmdList.SetViewport(0, 0, 0, DestTextureWidth, DestTextureHeight, 0.0f);
 		RHICmdList.SetScissorRect(false, 0, 0, 0, 0);
 
-		SetRenderTarget(RHICmdList, DestTexture, FTextureRHIRef());
-		
-		FGraphicsPipelineStateInitializer GraphicsPSOInit;
-		RHICmdList.ApplyCachedRenderTargets(GraphicsPSOInit);
-		GraphicsPSOInit.BlendState = TStaticBlendState<>::GetRHI();
-		GraphicsPSOInit.RasterizerState = TStaticRasterizerState<>::GetRHI();
-		GraphicsPSOInit.DepthStencilState = TStaticDepthStencilState<false, CF_Always>::GetRHI();
+		FRHIRenderPassInfo RPInfo(DestTexture, ERenderTargetActions::Load_Store);
+		RHICmdList.BeginRenderPass(RPInfo, TEXT("DownsampleRect"));
+		{
+			FGraphicsPipelineStateInitializer GraphicsPSOInit;
+			RHICmdList.ApplyCachedRenderTargets(GraphicsPSOInit);
+			GraphicsPSOInit.BlendState = TStaticBlendState<>::GetRHI();
+			GraphicsPSOInit.RasterizerState = TStaticRasterizerState<>::GetRHI();
+			GraphicsPSOInit.DepthStencilState = TStaticDepthStencilState<false, CF_Always>::GetRHI();
 
-		GraphicsPSOInit.BoundShaderState.VertexDeclarationRHI = RendererModule.GetFilterVertexDeclaration().VertexDeclarationRHI;
-		GraphicsPSOInit.BoundShaderState.VertexShaderRHI = GETSAFERHISHADER_VERTEX(*VertexShader);
-		GraphicsPSOInit.BoundShaderState.PixelShaderRHI = GETSAFERHISHADER_PIXEL(*PixelShader);
-		GraphicsPSOInit.PrimitiveType = PT_TriangleList;
-		SetGraphicsPipelineState(RHICmdList, GraphicsPSOInit);
-		
-		PixelShader->SetShaderParams(RHICmdList, FVector4(InvSrcTetureSize.X, InvSrcTetureSize.Y, 0, 0));
-		PixelShader->SetUVBounds(RHICmdList, FVector4(UVStart, UVEnd));
-		PixelShader->SetTexture(RHICmdList, Params.SourceTexture, BilinearClamp);
+			GraphicsPSOInit.BoundShaderState.VertexDeclarationRHI = RendererModule.GetFilterVertexDeclaration().VertexDeclarationRHI;
+			GraphicsPSOInit.BoundShaderState.VertexShaderRHI = GETSAFERHISHADER_VERTEX(*VertexShader);
+			GraphicsPSOInit.BoundShaderState.PixelShaderRHI = GETSAFERHISHADER_PIXEL(*PixelShader);
+			GraphicsPSOInit.PrimitiveType = PT_TriangleList;
+			SetGraphicsPipelineState(RHICmdList, GraphicsPSOInit);
 
-		RendererModule.DrawRectangle(
-			RHICmdList,
-			0, 0,
-			DownsampleSize.X, DownsampleSize.Y,
-			UVStart.X, UVStart.Y,
-			SizeUV.X, SizeUV.Y,
-			FIntPoint(DestTextureWidth, DestTextureHeight),
-			FIntPoint(1,1), 
-			*VertexShader,
-			EDRF_Default);
+			PixelShader->SetShaderParams(RHICmdList, FVector4(InvSrcTetureSize.X, InvSrcTetureSize.Y, 0, 0));
+			PixelShader->SetUVBounds(RHICmdList, FVector4(UVStart, UVEnd));
+			PixelShader->SetTexture(RHICmdList, Params.SourceTexture, BilinearClamp);
+
+			RendererModule.DrawRectangle(
+				RHICmdList,
+				0, 0,
+				DownsampleSize.X, DownsampleSize.Y,
+				UVStart.X, UVStart.Y,
+				SizeUV.X, SizeUV.Y,
+				FIntPoint(DestTextureWidth, DestTextureHeight),
+				FIntPoint(1, 1),
+				*VertexShader,
+				EDRF_Default);
+		}
+		RHICmdList.EndRenderPass();
 	}
 	
 	// Testing only
@@ -389,41 +401,47 @@ void FSlatePostProcessor::UpsampleRect(FRHICommandListImmediate& RHICmdList, IRe
 	RHICmdList.TransitionResource(EResourceTransitionAccess::EWritable, DestTexture);
 	RHICmdList.TransitionResource(EResourceTransitionAccess::EReadable, SrcTexture);
 
-	SetRenderTarget(RHICmdList, DestTexture, FTextureRHIRef());
-	RHICmdList.ApplyCachedRenderTargets(GraphicsPSOInit);
-
-	if (Params.RestoreStateFunc)
+	FRHIRenderPassInfo RPInfo(DestTexture, ERenderTargetActions::Load_Store);
+	RHICmdList.BeginRenderPass(RPInfo, TEXT("UpsampleRect"));
 	{
-		Params.RestoreStateFunc(RHICmdList, GraphicsPSOInit);
-	}
+		RHICmdList.ApplyCachedRenderTargets(GraphicsPSOInit);
 
-	TShaderMapRef<FScreenPS> PixelShader(ShaderMap);
+		if (Params.RestoreStateFunc)
+		{
+			// This can potentially end and restart a renderpass.
+			// #todo refactor so that we only start one renderpass here. Right now RestoreStateFunc may call UpdateScissorRect which requires an open renderpass.
+			Params.RestoreStateFunc(RHICmdList, GraphicsPSOInit);
+		}
 
-	GraphicsPSOInit.BoundShaderState.VertexDeclarationRHI = RendererModule.GetFilterVertexDeclaration().VertexDeclarationRHI;
-	GraphicsPSOInit.BoundShaderState.VertexShaderRHI = GETSAFERHISHADER_VERTEX(*VertexShader);
-	GraphicsPSOInit.BoundShaderState.PixelShaderRHI = GETSAFERHISHADER_PIXEL(*PixelShader);
-	GraphicsPSOInit.PrimitiveType = PT_TriangleList;
-	SetGraphicsPipelineState(RHICmdList, GraphicsPSOInit);
+		TShaderMapRef<FScreenPS> PixelShader(ShaderMap);
 
-	if (Params.RestoreStateFuncPostPipelineState)
-	{
-		Params.RestoreStateFuncPostPipelineState();
-	}
+		GraphicsPSOInit.BoundShaderState.VertexDeclarationRHI = RendererModule.GetFilterVertexDeclaration().VertexDeclarationRHI;
+		GraphicsPSOInit.BoundShaderState.VertexShaderRHI = GETSAFERHISHADER_VERTEX(*VertexShader);
+		GraphicsPSOInit.BoundShaderState.PixelShaderRHI = GETSAFERHISHADER_PIXEL(*PixelShader);
+		GraphicsPSOInit.PrimitiveType = PT_TriangleList;
+		SetGraphicsPipelineState(RHICmdList, GraphicsPSOInit);
 
-	PixelShader->SetParameters(RHICmdList, Sampler, SrcTexture);
+		if (Params.RestoreStateFuncPostPipelineState)
+		{
+			Params.RestoreStateFuncPostPipelineState();
+		}
 
-	const float SizeU = (DownsampledWidth == SrcTextureWidth) ? 1.0f : (DownsampledWidth / (float)SrcTextureWidth) - (1.0f / (float)SrcTextureWidth);
-	const float SizeV = (DownsampledHeight == SrcTextureHeight) ? 1.0f : (DownsampledHeight / (float)SrcTextureHeight) - (1.0f / (float)SrcTextureHeight);
+		PixelShader->SetParameters(RHICmdList, Sampler, SrcTexture);
 
-	RendererModule.DrawRectangle(RHICmdList,
-		DestRect.Left, DestRect.Top,
-		DestRect.Right - DestRect.Left, DestRect.Bottom - DestRect.Top,
-		0, 0,
-		SizeU, SizeV,
-		Params.SourceTextureSize,
-		FIntPoint(1, 1),
-		*VertexShader,
-		EDRF_Default);
+		const float SizeU = (DownsampledWidth == SrcTextureWidth) ? 1.0f : (DownsampledWidth / (float)SrcTextureWidth) - (1.0f / (float)SrcTextureWidth);
+		const float SizeV = (DownsampledHeight == SrcTextureHeight) ? 1.0f : (DownsampledHeight / (float)SrcTextureHeight) - (1.0f / (float)SrcTextureHeight);
+
+		RendererModule.DrawRectangle(RHICmdList,
+			DestRect.Left, DestRect.Top,
+			DestRect.Right - DestRect.Left, DestRect.Bottom - DestRect.Top,
+			0, 0,
+			SizeU, SizeV,
+			Params.SourceTextureSize,
+			FIntPoint(1, 1),
+			*VertexShader,
+			EDRF_Default);
+	}
+	RHICmdList.EndRenderPass();
 
 }
 #define BILINEAR_FILTER_METHOD 1
diff --git a/Engine/Source/Runtime/SlateRHIRenderer/Private/SlateRHIRenderer.cpp b/Engine/Source/Runtime/SlateRHIRenderer/Private/SlateRHIRenderer.cpp
--- a/Engine/Source/Runtime/SlateRHIRenderer/Private/SlateRHIRenderer.cpp
+++ b/Engine/Source/Runtime/SlateRHIRenderer/Private/SlateRHIRenderer.cpp
@@ -737,67 +737,68 @@ void FSlateRHIRenderer::DrawWindow_RenderThread(FRHICommandListImmediate& RHICmd
 				BackBuffer = ViewportInfo.UITargetRT;
 			}
 
-		if (SlateWireFrame)
-		{
-			bClear = true;
-		}
-
-		RHICmdList.BeginDrawingViewport( ViewportInfo.ViewportRHI, FTextureRHIRef() );
-		RHICmdList.SetViewport(0, 0, 0, ViewportWidth, ViewportHeight, 0.0f);
-		RHICmdList.TransitionResource(EResourceTransitionAccess::EWritable, BackBuffer);
-	
-		if (ViewportInfo.bRequiresStencilTest)
-		{
-			check(IsValidRef(ViewportInfo.DepthStencil));
-	
-			// Reset the backbuffer as our color render target and also set a depth stencil buffer
-			FRHIRenderTargetView ColorView(BackBuffer, 0, -1, bClear ? ERenderTargetLoadAction::EClear : ERenderTargetLoadAction::ELoad, ERenderTargetStoreAction::EStore);
-			FRHIDepthRenderTargetView DepthStencilView(ViewportInfo.DepthStencil, ERenderTargetLoadAction::ENoAction, ERenderTargetStoreAction::ENoAction, ERenderTargetLoadAction::ENoAction, ERenderTargetStoreAction::EStore);
-			FRHISetRenderTargetsInfo Info(1, &ColorView, DepthStencilView);
-	
-			// Clear the stencil buffer
-			RHICmdList.SetRenderTargetsAndClear(Info);
-		}
-		else
-		{
-			SetRenderTarget(RHICmdList, BackBuffer, FTextureRHIRef(), bClear ? ESimpleRenderTargetMode::EClearColorAndDepth : ESimpleRenderTargetMode::EExistingColorAndDepth);
-		}
-
-#if WITH_SLATE_VISUALIZERS
-		if (CVarShowSlateBatching.GetValueOnRenderThread() != 0 || CVarShowSlateOverdraw.GetValueOnRenderThread() != 0)
-		{
-			if (ViewportInfo.bRequiresStencilTest)
-			{
-				// Reset the backbuffer as our color render target and also set a depth stencil buffer
-				FRHIRenderTargetView ColorView(BackBuffer, 0, -1, ERenderTargetLoadAction::EClear, ERenderTargetStoreAction::EStore);
-				FRHISetRenderTargetsInfo Info(1, &ColorView, FRHIDepthRenderTargetView(ViewportInfo.DepthStencil, ERenderTargetLoadAction::ELoad, ERenderTargetStoreAction::EStore, ERenderTargetLoadAction::EClear, ERenderTargetStoreAction::EStore));
-				RHICmdList.SetRenderTargetsAndClear(Info);
-			}
-			else
+			if (SlateWireFrame)
 			{
-				SetRenderTarget(RHICmdList, BackBuffer, FTextureRHIRef(), ESimpleRenderTargetMode::EClearColorAndDepth);
+				bClear = true;
 			}
-		}
-#endif
 
-			if (BatchData.GetRenderBatches().Num() > 0)
+			RHICmdList.BeginDrawingViewport(ViewportInfo.ViewportRHI, FTextureRHIRef());
+			RHICmdList.SetViewport(0, 0, 0, ViewportWidth, ViewportHeight, 0.0f);
+			RHICmdList.TransitionResource(EResourceTransitionAccess::EWritable, BackBuffer);
+
 			{
-				SCOPE_CYCLE_COUNTER(STAT_SlateRTDrawBatches);
+				FRHIRenderPassInfo RPInfo(BackBuffer, ERenderTargetActions::Load_Store);
 
-				FSlateBackBuffer BackBufferTarget(BackBuffer, FIntPoint(ViewportWidth, ViewportHeight));
+				if (bClear)
+				{
+					RPInfo.ColorRenderTargets[0].Action = ERenderTargetActions::Clear_Store;
+				}
 
-				FSlateRenderingOptions DrawOptions(ViewMatrix * ViewportInfo.ProjectionMatrix);
-				DrawOptions.bWireFrame = !!SlateWireFrame;
+				if (ViewportInfo.bRequiresStencilTest)
+				{
+					check(IsValidRef(ViewportInfo.DepthStencil));
 
-				RenderingPolicy->DrawElements
-				(
-					RHICmdList,
-					BackBufferTarget,
-					BackBuffer,
-					ViewportInfo.DepthStencil,
-					BatchData.GetRenderBatches(),
-					DrawOptions
-				);
+					RPInfo.DepthStencilRenderTarget.Action = MakeDepthStencilTargetActions(ERenderTargetActions::DontLoad_DontStore, ERenderTargetActions::DontLoad_Store);
+					RPInfo.DepthStencilRenderTarget.DepthStencilTarget = ViewportInfo.DepthStencil;
+					RPInfo.DepthStencilRenderTarget.ExclusiveDepthStencil = FExclusiveDepthStencil::DepthNop_StencilWrite;
+				}
+
+	#if WITH_SLATE_VISUALIZERS
+				if (CVarShowSlateBatching.GetValueOnRenderThread() != 0 || CVarShowSlateOverdraw.GetValueOnRenderThread() != 0)
+				{
+					RPInfo.ColorRenderTargets[0].Action = ERenderTargetActions::Clear_Store;
+					if (ViewportInfo.bRequiresStencilTest)
+					{
+						// Reset the backbuffer as our color render target and also set a depth stencil buffer
+						RPInfo.DepthStencilRenderTarget.Action = MakeDepthStencilTargetActions(ERenderTargetActions::Load_Store, ERenderTargetActions::Clear_Store);
+						RPInfo.DepthStencilRenderTarget.DepthStencilTarget = ViewportInfo.DepthStencil;
+						RPInfo.DepthStencilRenderTarget.ExclusiveDepthStencil = FExclusiveDepthStencil::DepthWrite_StencilWrite;
+					}
+				}
+	#endif
+				RHICmdList.BeginRenderPass(RPInfo, TEXT("SlateBatches"));
+				{
+					if (BatchData.GetRenderBatches().Num() > 0)
+					{
+						SCOPE_CYCLE_COUNTER(STAT_SlateRTDrawBatches);
+
+						FSlateBackBuffer BackBufferTarget(BackBuffer, FIntPoint(ViewportWidth, ViewportHeight));
+
+						FSlateRenderingOptions DrawOptions(ViewMatrix * ViewportInfo.ProjectionMatrix);
+						DrawOptions.bWireFrame = !!SlateWireFrame;
+
+						RenderingPolicy->DrawElements
+						(
+							RHICmdList,
+							BackBufferTarget,
+							BackBuffer,
+							ViewportInfo.DepthStencil,
+							BatchData.GetRenderBatches(),
+							DrawOptions
+						);
+					}
+				}
+				RHICmdList.EndRenderPass();
 			}
 
 			if (bCompositeUI)
@@ -813,37 +814,39 @@ void FSlateRHIRenderer::DrawWindow_RenderThread(FRHICommandListImmediate& RHICmd
 				// Generate composition LUT
 				if (bLUTStale)
 				{
-					SetRenderTarget(RHICmdList, ViewportInfo.ColorSpaceLUTRT, FTextureRHIRef());
-
-					FGraphicsPipelineStateInitializer GraphicsPSOInit;
-					RHICmdList.ApplyCachedRenderTargets(GraphicsPSOInit);
-					GraphicsPSOInit.BlendState = TStaticBlendState<>::GetRHI();
-					GraphicsPSOInit.RasterizerState = TStaticRasterizerState<>::GetRHI();
-					GraphicsPSOInit.DepthStencilState = TStaticDepthStencilState<false, CF_Always>::GetRHI();
-
-					TShaderMapRef<FWriteToSliceVS> VertexShader(ShaderMap);
-					TOptionalShaderMapRef<FWriteToSliceGS> GeometryShader(ShaderMap);
-					TShaderMapRef<FCompositeLUTGenerationPS> PixelShader(ShaderMap);
-					const FVolumeBounds VolumeBounds(CompositionLUTSize);
-
-					GraphicsPSOInit.BoundShaderState.VertexDeclarationRHI = GScreenVertexDeclaration.VertexDeclarationRHI;
-					GraphicsPSOInit.BoundShaderState.VertexShaderRHI = GETSAFERHISHADER_VERTEX(*VertexShader);
-					GraphicsPSOInit.BoundShaderState.GeometryShaderRHI = GETSAFERHISHADER_GEOMETRY(*GeometryShader);
-					GraphicsPSOInit.BoundShaderState.PixelShaderRHI = GETSAFERHISHADER_PIXEL(*PixelShader);
-					GraphicsPSOInit.PrimitiveType = PT_TriangleStrip;
-					SetGraphicsPipelineState(RHICmdList, GraphicsPSOInit);
-
-					VertexShader->SetParameters(RHICmdList, VolumeBounds, FIntVector(VolumeBounds.MaxX - VolumeBounds.MinX));
-					if (GeometryShader.IsValid())
+					// #todo-renderpasses will this touch every pixel? use NoAction?
+					FRHIRenderPassInfo RPInfo(ViewportInfo.ColorSpaceLUTRT, ERenderTargetActions::Load_Store);
+					RHICmdList.BeginRenderPass(RPInfo, TEXT("GenerateLUT"));
 					{
-						GeometryShader->SetParameters(RHICmdList, VolumeBounds.MinZ);
-					}
-					PixelShader->SetParameters(RHICmdList);
+						FGraphicsPipelineStateInitializer GraphicsPSOInit;
+						RHICmdList.ApplyCachedRenderTargets(GraphicsPSOInit);
+						GraphicsPSOInit.BlendState = TStaticBlendState<>::GetRHI();
+						GraphicsPSOInit.RasterizerState = TStaticRasterizerState<>::GetRHI();
+						GraphicsPSOInit.DepthStencilState = TStaticDepthStencilState<false, CF_Always>::GetRHI();
+
+						TShaderMapRef<FWriteToSliceVS> VertexShader(ShaderMap);
+						TOptionalShaderMapRef<FWriteToSliceGS> GeometryShader(ShaderMap);
+						TShaderMapRef<FCompositeLUTGenerationPS> PixelShader(ShaderMap);
+						const FVolumeBounds VolumeBounds(CompositionLUTSize);
+
+						GraphicsPSOInit.BoundShaderState.VertexDeclarationRHI = GScreenVertexDeclaration.VertexDeclarationRHI;
+						GraphicsPSOInit.BoundShaderState.VertexShaderRHI = GETSAFERHISHADER_VERTEX(*VertexShader);
+						GraphicsPSOInit.BoundShaderState.GeometryShaderRHI = GETSAFERHISHADER_GEOMETRY(*GeometryShader);
+						GraphicsPSOInit.BoundShaderState.PixelShaderRHI = GETSAFERHISHADER_PIXEL(*PixelShader);
+						GraphicsPSOInit.PrimitiveType = PT_TriangleStrip;
+						SetGraphicsPipelineState(RHICmdList, GraphicsPSOInit);
 
-					RasterizeToVolumeTexture(RHICmdList, VolumeBounds);
+						VertexShader->SetParameters(RHICmdList, VolumeBounds, FIntVector(VolumeBounds.MaxX - VolumeBounds.MinX));
+						if (GeometryShader.IsValid())
+						{
+							GeometryShader->SetParameters(RHICmdList, VolumeBounds.MinZ);
+						}
+						PixelShader->SetParameters(RHICmdList);
 
-					FResolveParams ResolveParams;
-					RHICmdList.CopyToResolveTarget(ViewportInfo.ColorSpaceLUTRT, ViewportInfo.ColorSpaceLUTSRV, ResolveParams);
+						RasterizeToVolumeTexture(RHICmdList, VolumeBounds);
+					}
+					RHICmdList.EndRenderPass();
+					RHICmdList.CopyToResolveTarget(ViewportInfo.ColorSpaceLUTRT, ViewportInfo.ColorSpaceLUTSRV, FResolveParams());
 				}
 
 				// Composition pass
@@ -851,57 +854,60 @@ void FSlateRHIRenderer::DrawWindow_RenderThread(FRHICommandListImmediate& RHICmd
 					FResolveParams ResolveParams;
 					RHICmdList.CopyToResolveTarget(ViewportInfo.UITargetRT, ViewportInfo.UITargetSRV, ResolveParams);
 
-					SetRenderTarget(RHICmdList, FinalBuffer, FTextureRHIRef());
+					FRHIRenderPassInfo RPInfo(FinalBuffer, ERenderTargetActions::Load_Store);
+					RHICmdList.BeginRenderPass(RPInfo, TEXT("SlateComposite"));
+					{
+						FGraphicsPipelineStateInitializer GraphicsPSOInit;
+						RHICmdList.ApplyCachedRenderTargets(GraphicsPSOInit);
+						GraphicsPSOInit.BlendState = TStaticBlendState<>::GetRHI();
+						GraphicsPSOInit.RasterizerState = TStaticRasterizerState<>::GetRHI();
+						GraphicsPSOInit.DepthStencilState = TStaticDepthStencilState<false, CF_Always>::GetRHI();
 
-					FGraphicsPipelineStateInitializer GraphicsPSOInit;
-					RHICmdList.ApplyCachedRenderTargets(GraphicsPSOInit);
-					GraphicsPSOInit.BlendState = TStaticBlendState<>::GetRHI();
-					GraphicsPSOInit.RasterizerState = TStaticRasterizerState<>::GetRHI();
-					GraphicsPSOInit.DepthStencilState = TStaticDepthStencilState<false, CF_Always>::GetRHI();
+						TShaderMapRef<FScreenVS> VertexShader(ShaderMap);
 
-					TShaderMapRef<FScreenVS> VertexShader(ShaderMap);
+						if (HDROutputDevice == 5 || HDROutputDevice == 6)
+						{
+							// ScRGB encoding
+							TShaderMapRef<FCompositePS<1>> PixelShader(ShaderMap);
 
-					if (HDROutputDevice == 5 || HDROutputDevice == 6)
-					{
-						// ScRGB encoding
-						TShaderMapRef<FCompositePS<1>> PixelShader(ShaderMap);
+							GraphicsPSOInit.BoundShaderState.VertexDeclarationRHI = RendererModule.GetFilterVertexDeclaration().VertexDeclarationRHI;
+							GraphicsPSOInit.BoundShaderState.VertexShaderRHI = GETSAFERHISHADER_VERTEX(*VertexShader);
+							GraphicsPSOInit.BoundShaderState.PixelShaderRHI = GETSAFERHISHADER_PIXEL(*PixelShader);
+							GraphicsPSOInit.PrimitiveType = PT_TriangleList;
 
-						GraphicsPSOInit.BoundShaderState.VertexDeclarationRHI = RendererModule.GetFilterVertexDeclaration().VertexDeclarationRHI;
-						GraphicsPSOInit.BoundShaderState.VertexShaderRHI = GETSAFERHISHADER_VERTEX(*VertexShader);
-						GraphicsPSOInit.BoundShaderState.PixelShaderRHI = GETSAFERHISHADER_PIXEL(*PixelShader);
-						GraphicsPSOInit.PrimitiveType = PT_TriangleList;
+							SetGraphicsPipelineState(RHICmdList, GraphicsPSOInit);
 
-						SetGraphicsPipelineState(RHICmdList, GraphicsPSOInit);
+							PixelShader->SetParameters(RHICmdList, ViewportInfo.UITargetSRV, ViewportInfo.HDRSourceSRV, ViewportInfo.ColorSpaceLUTSRV);
+						}
+						else
+						{
+							// ST2084 (PQ) encoding
+							TShaderMapRef<FCompositePS<0>> PixelShader(ShaderMap);
 
-						PixelShader->SetParameters(RHICmdList, ViewportInfo.UITargetSRV, ViewportInfo.HDRSourceSRV, ViewportInfo.ColorSpaceLUTSRV);
-					}
-					else
-					{
-						// ST2084 (PQ) encoding
-						TShaderMapRef<FCompositePS<0>> PixelShader(ShaderMap);
+							GraphicsPSOInit.BoundShaderState.VertexDeclarationRHI = RendererModule.GetFilterVertexDeclaration().VertexDeclarationRHI;
+							GraphicsPSOInit.BoundShaderState.VertexShaderRHI = GETSAFERHISHADER_VERTEX(*VertexShader);
+							GraphicsPSOInit.BoundShaderState.PixelShaderRHI = GETSAFERHISHADER_PIXEL(*PixelShader);
+							GraphicsPSOInit.PrimitiveType = PT_TriangleList;
 
-						GraphicsPSOInit.BoundShaderState.VertexDeclarationRHI = RendererModule.GetFilterVertexDeclaration().VertexDeclarationRHI;
-						GraphicsPSOInit.BoundShaderState.VertexShaderRHI = GETSAFERHISHADER_VERTEX(*VertexShader);
-						GraphicsPSOInit.BoundShaderState.PixelShaderRHI = GETSAFERHISHADER_PIXEL(*PixelShader);
-						GraphicsPSOInit.PrimitiveType = PT_TriangleList;
+							SetGraphicsPipelineState(RHICmdList, GraphicsPSOInit);
 
-						SetGraphicsPipelineState(RHICmdList, GraphicsPSOInit);
+							PixelShader->SetParameters(RHICmdList, ViewportInfo.UITargetSRV, ViewportInfo.HDRSourceSRV, ViewportInfo.ColorSpaceLUTSRV);
+						}
 
-						PixelShader->SetParameters(RHICmdList, ViewportInfo.UITargetSRV, ViewportInfo.HDRSourceSRV, ViewportInfo.ColorSpaceLUTSRV);
+						RendererModule.DrawRectangle(
+							RHICmdList,
+							0, 0,
+							ViewportWidth, ViewportHeight,
+							0, 0,
+							ViewportWidth, ViewportHeight,
+							FIntPoint(ViewportWidth, ViewportHeight),
+							FIntPoint(ViewportWidth, ViewportHeight),
+							*VertexShader,
+							EDRF_UseTriangleOptimization);
 					}
-
-					RendererModule.DrawRectangle(
-						RHICmdList,
-						0, 0,
-						ViewportWidth, ViewportHeight,
-						0, 0,
-						ViewportWidth, ViewportHeight,
-						FIntPoint(ViewportWidth, ViewportHeight),
-						FIntPoint(ViewportWidth, ViewportHeight),
-						*VertexShader,
-						EDRF_UseTriangleOptimization);
+					RHICmdList.EndRenderPass();
 				}
-			}
+			} //bCompositeUI
 
 			if (!bRenderedStereo && GEngine && IsValidRef(ViewportInfo.GetRenderTargetTexture()) && GEngine->StereoRenderingDevice.IsValid())
 			{
diff --git a/Engine/Source/Runtime/SlateRHIRenderer/Private/SlateRHIRenderingPolicy.cpp b/Engine/Source/Runtime/SlateRHIRenderer/Private/SlateRHIRenderingPolicy.cpp
--- a/Engine/Source/Runtime/SlateRHIRenderer/Private/SlateRHIRenderingPolicy.cpp
+++ b/Engine/Source/Runtime/SlateRHIRenderer/Private/SlateRHIRenderingPolicy.cpp
@@ -304,7 +304,7 @@ static FSceneView* CreateSceneView( FSceneViewFamilyContext* ViewFamilyContext,
 
 static const FName RendererModuleName("Renderer");
 
-static void UpdateScissorRect(
+static bool UpdateScissorRect(
 	FRHICommandList& RHICmdList, 
 #if STATS
 	int32& ScissorClips, 
@@ -323,6 +323,9 @@ static void UpdateScissorRect(
 	const FMatrix& ViewProjection, 
 	bool bForceStateChange)
 {
+	check(RHICmdList.IsInsideRenderPass());
+	bool bDidRestartRenderpass = false;
+
 	QUICK_SCOPE_CYCLE_COUNTER(STAT_Slate_UpdateScissorRect);
 
 	if (RenderBatch.ClippingState != LastClippingState || bForceStateChange)
@@ -338,10 +341,16 @@ static void UpdateScissorRect(
 
 				if (bForceStateChange && MaskingID > 0)
 				{
-					FRHIRenderTargetView ColorView(ColorTarget, ERenderTargetLoadAction::ELoad);
-					FRHIDepthRenderTargetView DepthStencilView(DepthStencilTarget, ERenderTargetLoadAction::ENoAction, ERenderTargetStoreAction::ENoAction, ERenderTargetLoadAction::ELoad, ERenderTargetStoreAction::EStore);
-					FRHISetRenderTargetsInfo CurrentRTInfo(1, &ColorView, DepthStencilView);
-					RHICmdList.SetRenderTargetsAndClear(CurrentRTInfo);
+					// #todo-renderpasses this is very gross. If/when this gets refactored we can detect a simple clear or batch up elements by rendertarget (and other stuff)
+					RHICmdList.EndRenderPass();
+					bDidRestartRenderpass = true;
+
+					FRHIRenderPassInfo RPInfo(ColorTarget, ERenderTargetActions::Load_Store);
+					RPInfo.DepthStencilRenderTarget.Action = MakeDepthStencilTargetActions(ERenderTargetActions::DontLoad_DontStore, ERenderTargetActions::Load_Store);
+					RPInfo.DepthStencilRenderTarget.DepthStencilTarget = DepthStencilTarget;
+					RPInfo.DepthStencilRenderTarget.ExclusiveDepthStencil = FExclusiveDepthStencil::DepthNop_StencilWrite;
+
+					RHICmdList.BeginRenderPass(RPInfo, TEXT("SlateUpdateScissorRect"));
 				}
 
 				const FSlateClippingZone& ScissorRect = ClipState.ScissorRect.GetValue();
@@ -423,11 +432,18 @@ static void UpdateScissorRect(
 				// Don't bother setting the render targets unless we actually need to clear them.
 				if (bClearStencil || bForceStateChange)
 				{
+					// #todo-renderpasses Similar to above this is gross. Would require a refactor to really fix.
+					RHICmdList.EndRenderPass();
+					bDidRestartRenderpass = true;
+
 					// Clear current stencil buffer, we use ELoad/EStore, because we need to keep the stencil around.
-					FRHIRenderTargetView ColorView(ColorTarget, ERenderTargetLoadAction::ELoad);
-					FRHIDepthRenderTargetView DepthStencilView(DepthStencilTarget, ERenderTargetLoadAction::ENoAction, ERenderTargetStoreAction::ENoAction, bClearStencil ? ERenderTargetLoadAction::EClear : ERenderTargetLoadAction::ELoad, ERenderTargetStoreAction::EStore);
-					FRHISetRenderTargetsInfo CurrentRTInfo(1, &ColorView, DepthStencilView);
-					RHICmdList.SetRenderTargetsAndClear(CurrentRTInfo);
+					ERenderTargetLoadAction StencilLoadAction = bClearStencil ? ERenderTargetLoadAction::EClear : ERenderTargetLoadAction::ELoad;
+
+					FRHIRenderPassInfo RPInfo(ColorTarget, ERenderTargetActions::Load_Store);
+					RPInfo.DepthStencilRenderTarget.Action = MakeDepthStencilTargetActions(ERenderTargetActions::DontLoad_DontStore, MakeRenderTargetActions(StencilLoadAction, ERenderTargetStoreAction::EStore));
+					RPInfo.DepthStencilRenderTarget.DepthStencilTarget = DepthStencilTarget;
+					RPInfo.DepthStencilRenderTarget.ExclusiveDepthStencil = FExclusiveDepthStencil::DepthNop_StencilWrite;
+					RHICmdList.BeginRenderPass(RPInfo, TEXT("SlateUpdateScissorRect_ClearStencil"));
 				}
 
 				// Start by setting up the stenciling states so that we can write representations of the clipping zones into the stencil buffer only.
@@ -589,6 +605,8 @@ static void UpdateScissorRect(
 
 		LastClippingState = RenderBatch.ClippingState;
 	}
+
+	return bDidRestartRenderpass;
 }
 
 void FSlateRHIRenderingPolicy::DrawElements(
@@ -601,6 +619,7 @@ void FSlateRHIRenderingPolicy::DrawElements(
 {
 	// Should only be called by the rendering thread
 	check(IsInRenderingThread());
+	check(RHICmdList.IsInsideRenderPass());
 
 	// Cache the TextureLODGroups so that we can look them up for texture filtering.
 	if (UDeviceProfileManager::DeviceProfileManagerSingleton)
@@ -710,8 +729,13 @@ void FSlateRHIRenderingPolicy::DrawElements(
 	FVector2D ViewTranslation2D = Options.ViewOffset;
 
 	// Draw each element
+	/*
+		#todo-renderpasses This loop ends up with ugly logic.
+		CustomDrawers will draw in their own renderpass. So we must remember to reopen the renderpass with the passed in Color/DepthStencil targets.
+	*/
 	for (int32 BatchIndex = 0; BatchIndex < RenderBatches.Num(); ++BatchIndex)
 	{
+		check(RHICmdList.IsInsideRenderPass());
 #if WITH_SLATE_VISUALIZERS
 		FLinearColor BatchColor = FLinearColor(BatchColors.GetUnitVector());
 #endif
@@ -798,6 +822,7 @@ void FSlateRHIRenderingPolicy::DrawElements(
 			ESlateShaderResource::Type ResourceType = ShaderResource ? ShaderResource->GetType() : ESlateShaderResource::Invalid;
 			if (ResourceType != ESlateShaderResource::Material && ShaderType != ESlateShader::PostProcess)
 			{
+				check(RHICmdList.IsInsideRenderPass());
 				check(RenderBatch.NumIndices > 0);
 				FSlateElementPS* PixelShader = nullptr;
 
@@ -1011,6 +1036,7 @@ void FSlateRHIRenderingPolicy::DrawElements(
 			else if (GEngine && ShaderResource && ShaderResource->GetType() == ESlateShaderResource::Material && ShaderType != ESlateShader::PostProcess)
 			{
 				SLATE_DRAW_EVENT(RHICmdList, MaterialBatch);
+				check(RHICmdList.IsInsideRenderPass());
 
 				check(RenderBatch.NumIndices > 0);
 				// Note: This code is only executed if the engine is loaded (in early loading screens attempting to use a material is unsupported
@@ -1171,6 +1197,7 @@ void FSlateRHIRenderingPolicy::DrawElements(
 			else if (ShaderType == ESlateShader::PostProcess)
 			{
 				SLATE_DRAW_EVENT(RHICmdList, PostProcess);
+				RHICmdList.EndRenderPass();
 
 				const FVector4& QuadPositionData = ShaderParams.PixelParams;
 
@@ -1181,7 +1208,7 @@ void FSlateRHIRenderingPolicy::DrawElements(
 				RectParams.SourceTextureSize = BackBuffer.GetSizeXY();
 
 				RectParams.RestoreStateFunc = [&](FRHICommandListImmediate&InRHICmdList, FGraphicsPipelineStateInitializer& InGraphicsPSOInit) {
-					UpdateScissorRect(
+					return UpdateScissorRect(
 						InRHICmdList,
 #if STATS
 						ScissorClips,
@@ -1211,14 +1238,33 @@ void FSlateRHIRenderingPolicy::DrawElements(
 				BlurParams.DownsampleAmount = ShaderParams.PixelParams2.Z;
 
 				PostProcessor->BlurRect(RHICmdList, RendererModule, BlurParams, RectParams);
-			}
 
+				check(RHICmdList.IsOutsideRenderPass());
+				{
+					FRHIRenderPassInfo RPInfo(BackBuffer.GetRenderTargetTexture(), ERenderTargetActions::Load_Store);
+					RPInfo.DepthStencilRenderTarget.DepthStencilTarget = DepthStencilTarget;
+					if (DepthStencilTarget)
+					{
+						RPInfo.DepthStencilRenderTarget.Action = EDepthStencilTargetActions::LoadDepthStencil_StoreDepthStencil;
+						RPInfo.DepthStencilRenderTarget.ExclusiveDepthStencil = FExclusiveDepthStencil::DepthWrite_StencilWrite;
+					}
+					else
+					{
+						RPInfo.DepthStencilRenderTarget.Action = EDepthStencilTargetActions::DontLoad_DontStore;
+						RPInfo.DepthStencilRenderTarget.ExclusiveDepthStencil = FExclusiveDepthStencil::DepthNop_StencilNop;
+					}
+					RHICmdList.BeginRenderPass(RPInfo, TEXT("RestartingSlateDrawElements"));
+				}
+			}
 		}
 		else
 		{
 			TSharedPtr<ICustomSlateElement, ESPMode::ThreadSafe> CustomDrawer = RenderBatch.CustomDrawer.Pin();
 			if (CustomDrawer.IsValid())
 			{
+				// CustomDrawers will change the rendertarget. So we must close any outstanding renderpasses.
+				RHICmdList.EndRenderPass();	
+
 				SLATE_DRAW_EVENT(RHICmdList, CustomDrawer);
 
 				// Disable scissor rect. A previous draw element may have had one
@@ -1228,15 +1274,29 @@ void FSlateRHIRenderingPolicy::DrawElements(
 				// This element is custom and has no Slate geometry.  Tell it to render itself now
 				CustomDrawer->DrawRenderThread(RHICmdList, &BackBuffer.GetRenderTargetTexture());
 
-
 				//We reset the maskingID here because otherwise the RT might not get re-set in the lines above see: if (bClearStencil || bForceStateChange)
 				MaskingID = 0;
 
+				// Restart the renderpass since the CustomDrawer will have drawn into its own.
+				FRHIRenderPassInfo RPInfo(ColorTarget, ERenderTargetActions::Load_Store);
+				RPInfo.DepthStencilRenderTarget.DepthStencilTarget = DepthStencilTarget;
+				if (DepthStencilTarget)
+				{
+					RPInfo.DepthStencilRenderTarget.Action = EDepthStencilTargetActions::LoadDepthStencil_StoreDepthStencil;
+					RPInfo.DepthStencilRenderTarget.ExclusiveDepthStencil = FExclusiveDepthStencil::DepthWrite_StencilWrite;
+				}
+				else
+				{
+					RPInfo.DepthStencilRenderTarget.Action = EDepthStencilTargetActions::DontLoad_DontStore;
+					RPInfo.DepthStencilRenderTarget.ExclusiveDepthStencil = FExclusiveDepthStencil::DepthNop_StencilNop;
+				}
+				RHICmdList.BeginRenderPass(RPInfo, TEXT("SlateDrawElements"));
+
 				// Something may have messed with the viewport size so set it back to the full target.
 				RHICmdList.SetViewport(0, 0, 0, BackBuffer.GetSizeXY().X, BackBuffer.GetSizeXY().Y, 0.0f);
 				RHICmdList.SetStreamSource(0, VertexBuffer->VertexBufferRHI, 0);
 			}
-		}
+		} // CustomDrawer
 	}
 
 	// Don't do color correction on iOS or Android, we don't have the GPU overhead for it.
