DR - Copy from 4567116
- Fix engine test failing on debug/release and on Vulkan
- Parameters on shaders must be set again after calling SetPSO()

diff --git a/Engine/Source/Runtime/SlateRHIRenderer/Private/SlateRHIRenderingPolicy.cpp b/Engine/Source/Runtime/SlateRHIRenderer/Private/SlateRHIRenderingPolicy.cpp
--- a/Engine/Source/Runtime/SlateRHIRenderer/Private/SlateRHIRenderingPolicy.cpp
+++ b/Engine/Source/Runtime/SlateRHIRenderer/Private/SlateRHIRenderingPolicy.cpp
@@ -533,6 +533,9 @@ static bool UpdateScissorRect(
 
 
 						SetGraphicsPipelineState(RHICmdList, WriteMaskPSOInit);
+
+						VertexShader->SetViewProjection(RHICmdList, ViewProjection);
+						VertexShader->SetVerticalAxisMultiplier(RHICmdList, bSwitchVerticalAxis ? -1.0f : 1.0f);
 					}
 				}
 
