Fix access to scene textures from slate materials

diff --git a/Engine/Source/Runtime/SlateRHIRenderer/Private/SlateMaterialShader.cpp b/Engine/Source/Runtime/SlateRHIRenderer/Private/SlateMaterialShader.cpp
--- a/Engine/Source/Runtime/SlateRHIRenderer/Private/SlateMaterialShader.cpp
+++ b/Engine/Source/Runtime/SlateRHIRenderer/Private/SlateMaterialShader.cpp
@@ -118,7 +118,8 @@ void FSlateMaterialShaderPS::SetParameters(FRHICommandList& RHICmdList, const FS
 
 	SetShaderValue( RHICmdList, ShaderRHI, ShaderParams, InShaderParams );
 
-	FMaterialShader::SetParameters<FPixelShaderRHIParamRef>(RHICmdList, ShaderRHI, MaterialRenderProxy, *Material, View, View.ViewUniformBuffer, ESceneTextureSetupMode::None);
+	const ESceneTextureSetupMode SceneTextures = ESceneTextureSetupMode::SceneDepth | ESceneTextureSetupMode::SSAO | ESceneTextureSetupMode::CustomDepth; // TODO - SSAO valid here?
+	FMaterialShader::SetParameters<FPixelShaderRHIParamRef>(RHICmdList, ShaderRHI, MaterialRenderProxy, *Material, View, View.ViewUniformBuffer, SceneTextures);
 }
 
 void FSlateMaterialShaderPS::SetAdditionalTexture( FRHICommandList& RHICmdList, const FTextureRHIParamRef InTexture, const FSamplerStateRHIRef SamplerState )
