Fixes color deficiency's handling of renderpasses. Pointed out by Per-Jonny Kack.

diff --git a/Engine/Source/Runtime/SlateRHIRenderer/Private/SlateRHIRenderer.cpp b/Engine/Source/Runtime/SlateRHIRenderer/Private/SlateRHIRenderer.cpp
--- a/Engine/Source/Runtime/SlateRHIRenderer/Private/SlateRHIRenderer.cpp
+++ b/Engine/Source/Runtime/SlateRHIRenderer/Private/SlateRHIRenderer.cpp
@@ -802,7 +802,15 @@ void FSlateRHIRenderer::DrawWindow_RenderThread(FRHICommandListImmediate& RHICmd
 						);
 					}
 				}
-				RHICmdList.EndRenderPass();
+
+				// @todo Could really use a refactor.
+				// Kind of gross but we don't want to restart renderpasses for no reason.
+				// If the color deficiency shaders are active within DrawElements there will not be a renderpass here.
+				// In the general case there will be a RenderPass active at this point.
+				if (RHICmdList.IsInsideRenderPass())
+				{
+					RHICmdList.EndRenderPass();
+				}
 			}
 
 			if (bCompositeUI)
diff --git a/Engine/Source/Runtime/SlateRHIRenderer/Private/SlateRHIRenderingPolicy.cpp b/Engine/Source/Runtime/SlateRHIRenderer/Private/SlateRHIRenderingPolicy.cpp
--- a/Engine/Source/Runtime/SlateRHIRenderer/Private/SlateRHIRenderingPolicy.cpp
+++ b/Engine/Source/Runtime/SlateRHIRenderer/Private/SlateRHIRenderingPolicy.cpp
@@ -1306,6 +1306,8 @@ void FSlateRHIRenderingPolicy::DrawElements(
 #if !(PLATFORM_IOS || PLATFORM_ANDROID)
 	if (GSlateColorDeficiencyType != EColorVisionDeficiency::NormalVision && GSlateColorDeficiencySeverity > 0)
 	{
+		RHICmdList.EndRenderPass();
+
 		FPostProcessRectParams RectParams;
 		RectParams.SourceTexture = BackBuffer.GetRenderTargetTexture();
 		RectParams.SourceRect = FSlateRect(0, 0, BackBuffer.GetSizeXY().X, BackBuffer.GetSizeXY().Y);
@@ -1313,6 +1315,13 @@ void FSlateRHIRenderingPolicy::DrawElements(
 		RectParams.SourceTextureSize = BackBuffer.GetSizeXY();
 
 		PostProcessor->ColorDeficiency(RHICmdList, RendererModule, RectParams);
+
+		FRHIRenderPassInfo RPInfo(ColorTarget, ERenderTargetActions::Load_Store);
+		RPInfo.DepthStencilRenderTarget.DepthStencilTarget = DepthStencilTarget;
+
+		// @todo refactor this.
+		// ColorDeficiency has self-contained renderpasses. To avoid starting an empty renderpass we do not
+		// restart the renderpass here.
 	}
 #endif
 
