Round both to ints to avoid having to do an extra resize by being +/- 1 off. Due to one ceiling and one trunking the int

diff --git a/Engine/Source/Runtime/SlateRHIRenderer/Private/SlateRHIRenderer.cpp b/Engine/Source/Runtime/SlateRHIRenderer/Private/SlateRHIRenderer.cpp
--- a/Engine/Source/Runtime/SlateRHIRenderer/Private/SlateRHIRenderer.cpp
+++ b/Engine/Source/Runtime/SlateRHIRenderer/Private/SlateRHIRenderer.cpp
@@ -276,8 +276,8 @@ void FSlateRHIRenderer::CreateViewport(const TSharedRef<SWindow> Window)
 
 		// Clamp the window size to a reasonable default anything below 8 is a d3d warning and 8 is used anyway.
 		// @todo Slate: This is a hack to work around menus being summoned with 0,0 for window size until they are ticked.
-		int32 Width = FMath::Max(MIN_VIEWPORT_SIZE, FMath::TruncToInt(WindowSize.X));
-		int32 Height = FMath::Max(MIN_VIEWPORT_SIZE, FMath::TruncToInt(WindowSize.Y));
+		int32 Width = FMath::Max(MIN_VIEWPORT_SIZE, FMath::CeilToInt(WindowSize.X));
+		int32 Height = FMath::Max(MIN_VIEWPORT_SIZE, FMath::CeilToInt(WindowSize.Y));
 
 		// Sanity check dimensions
 		if (!ensureMsgf(Width <= MAX_VIEWPORT_SIZE && Height <= MAX_VIEWPORT_SIZE, TEXT("Invalid window with Width=%u and Height=%u"), Width, Height))
