Copy array into FSlateUpdateInstanceBufferCommand instead of passing a reference to fix a potential race condition due to basic Free Buffer Array index rotating

#rnx

diff --git a/Engine/Source/Runtime/SlateRHIRenderer/Private/SlateUpdatableBuffer.cpp b/Engine/Source/Runtime/SlateRHIRenderer/Private/SlateUpdatableBuffer.cpp
--- a/Engine/Source/Runtime/SlateRHIRenderer/Private/SlateUpdatableBuffer.cpp
+++ b/Engine/Source/Runtime/SlateRHIRenderer/Private/SlateUpdatableBuffer.cpp
@@ -8,7 +8,7 @@ DECLARE_CYCLE_STAT(TEXT("UpdateInstanceBuffer Time"), STAT_SlateUpdateInstanceBu
 struct FSlateUpdateInstanceBufferCommand final : public FRHICommand<FSlateUpdateInstanceBufferCommand>
 {
 	TSlateElementVertexBuffer<FVector4>& InstanceBuffer;
-	const TArray<FVector4>& InstanceData;
+	TArray<FVector4> InstanceData;
 
 	FSlateUpdateInstanceBufferCommand(TSlateElementVertexBuffer<FVector4>& InInstanceBuffer, const TArray<FVector4>& InInstanceData )
 		: InstanceBuffer(InInstanceBuffer)
