Fixed wrong size passed to lock in FSlateUpdateInstanceBufferCommand
Removed FlushRenderingCommands in destructor and defer destruction of BufferData and VertexBuffer
Keep a reference of VertexBufferRHI in RHI command in case VertexBufferRHI gets modified on RT by being resized

#rnx

diff --git a/Engine/Source/Runtime/SlateRHIRenderer/Private/SlateUpdatableBuffer.cpp b/Engine/Source/Runtime/SlateRHIRenderer/Private/SlateUpdatableBuffer.cpp
--- a/Engine/Source/Runtime/SlateRHIRenderer/Private/SlateUpdatableBuffer.cpp
+++ b/Engine/Source/Runtime/SlateRHIRenderer/Private/SlateUpdatableBuffer.cpp
@@ -9,10 +9,12 @@ struct FSlateUpdateInstanceBufferCommand final : public FRHICommand<FSlateUpdate
 {
 	TSlateElementVertexBuffer<FVector4>& InstanceBuffer;
 	TArray<FVector4> InstanceData;
+	FVertexBufferRHIRef VertexBuffer;
 
-	FSlateUpdateInstanceBufferCommand(TSlateElementVertexBuffer<FVector4>& InInstanceBuffer, const TArray<FVector4>& InInstanceData )
+	FSlateUpdateInstanceBufferCommand(TSlateElementVertexBuffer<FVector4>& InInstanceBuffer, const TArray<FVector4>& InInstanceData)
 		: InstanceBuffer(InInstanceBuffer)
 		, InstanceData(InInstanceData)
+		, VertexBuffer(InInstanceBuffer.VertexBufferRHI) // Keep a reference of RHIVertexBuffer in case it changes on RT thread during a resize
 	{}
 
 	void Execute(FRHICommandListBase& CmdList)
@@ -20,8 +22,7 @@ struct FSlateUpdateInstanceBufferCommand final : public FRHICommand<FSlateUpdate
 		SCOPE_CYCLE_COUNTER( STAT_SlateUpdateInstanceBuffer );
 		const bool bIsInRenderingThread = !IsRunningRHIInSeparateThread() || CmdList.Bypass();
 
-		int32 RequiredVertexBufferSize = InstanceData.Num()*sizeof(FVector4);
-		uint8* InstanceBufferData = (uint8*)InstanceBuffer.LockBuffer(RequiredVertexBufferSize, bIsInRenderingThread);
+		uint8* InstanceBufferData = (uint8*)InstanceBuffer.LockBuffer(InstanceData.Num(), bIsInRenderingThread);
 
 		FMemory::Memcpy( InstanceBufferData, InstanceData.GetData(), InstanceData.Num()*sizeof(FVector4) );
 	
@@ -32,23 +33,28 @@ struct FSlateUpdateInstanceBufferCommand final : public FRHICommand<FSlateUpdate
 FSlateUpdatableInstanceBuffer::FSlateUpdatableInstanceBuffer( int32 InitialInstanceCount )
 	: FreeBufferIndex(0)
 {
-	InstanceBufferResource.Init(InitialInstanceCount);
+	InstanceData = new FInstanceData;
+
+	InstanceData->InstanceBufferResource.Init(InitialInstanceCount);
 	for( int32 BufferIndex = 0; BufferIndex < SlateRHIConstants::NumBuffers; ++BufferIndex )
 	{
-		BufferData[BufferIndex].Reserve( InitialInstanceCount );
+		InstanceData->Array[BufferIndex].Reserve( InitialInstanceCount );
 	}
-
 }
 
 FSlateUpdatableInstanceBuffer::~FSlateUpdatableInstanceBuffer()
 {
-	InstanceBufferResource.Destroy();
-	FlushRenderingCommands();
+	InstanceData->InstanceBufferResource.Destroy();
+	BeginCleanup(InstanceData);
+}
+
+FSlateUpdatableInstanceBuffer::FInstanceData::~FInstanceData()
+{
 }
 
 void FSlateUpdatableInstanceBuffer::BindStreamSource(FRHICommandListImmediate& RHICmdList, int32 StreamIndex, uint32 InstanceOffset)
 {
-	RHICmdList.SetStreamSource(StreamIndex, InstanceBufferResource.VertexBufferRHI, InstanceOffset*sizeof(FVector4));
+	RHICmdList.SetStreamSource(StreamIndex, InstanceData->InstanceBufferResource.VertexBufferRHI, InstanceOffset*sizeof(FVector4));
 }
 
 TSharedPtr<class FSlateInstanceBufferUpdate> FSlateUpdatableInstanceBuffer::BeginUpdate()
@@ -81,24 +87,24 @@ void FSlateUpdatableInstanceBuffer::UpdateRenderingData(int32 NumInstancesToUse)
 
 TArray<FVector4>& FSlateUpdatableInstanceBuffer::GetBufferData()
 {
-	return BufferData[FreeBufferIndex];
+	return InstanceData->Array[FreeBufferIndex];
 }
 
 void FSlateUpdatableInstanceBuffer::UpdateRenderingData_RenderThread(FRHICommandListImmediate& RHICmdList, int32 BufferIndex)
 {
 	SCOPE_CYCLE_COUNTER( STAT_SlateUpdateInstanceBuffer );
 
-	const TArray<FVector4>& RenderThreadBufferData = BufferData[BufferIndex];
-	InstanceBufferResource.PreFillBuffer( RenderThreadBufferData.Num(), false );
+	const TArray<FVector4>& RenderThreadBufferData = InstanceData->Array[BufferIndex];
+	InstanceData->InstanceBufferResource.PreFillBuffer( RenderThreadBufferData.Num(), false );
 
 	if(!IsRunningRHIInSeparateThread() || RHICmdList.Bypass())
 	{
-		FSlateUpdateInstanceBufferCommand Command(InstanceBufferResource, RenderThreadBufferData);
+		FSlateUpdateInstanceBufferCommand Command(InstanceData->InstanceBufferResource, RenderThreadBufferData);
 		Command.Execute(RHICmdList);
 	}
 	else
 	{
-		ALLOC_COMMAND_CL(RHICmdList, FSlateUpdateInstanceBufferCommand)(InstanceBufferResource, RenderThreadBufferData);
+		ALLOC_COMMAND_CL(RHICmdList, FSlateUpdateInstanceBufferCommand)(InstanceData->InstanceBufferResource, RenderThreadBufferData);
 	}
 }
 
diff --git a/Engine/Source/Runtime/SlateRHIRenderer/Private/SlateUpdatableBuffer.h b/Engine/Source/Runtime/SlateRHIRenderer/Private/SlateUpdatableBuffer.h
--- a/Engine/Source/Runtime/SlateRHIRenderer/Private/SlateUpdatableBuffer.h
+++ b/Engine/Source/Runtime/SlateRHIRenderer/Private/SlateUpdatableBuffer.h
@@ -27,9 +27,22 @@ private:
 	// END ISlateUpdatableInstanceBuffer
 
 	void UpdateRenderingData_RenderThread(FRHICommandListImmediate& RHICmdList, int32 BufferIndex);
+
 private:
-	TArray<FVector4> BufferData[SlateRHIConstants::NumBuffers];
-	TSlateElementVertexBuffer<FVector4> InstanceBufferResource;
+	class FInstanceData : public FDeferredCleanupInterface
+	{
+	public:
+		FInstanceData()
+		{
+		}
+
+		virtual ~FInstanceData();
+
+		TArray<FVector4> Array[SlateRHIConstants::NumBuffers];
+		TSlateElementVertexBuffer<FVector4> InstanceBufferResource;
+	};
+
+	FInstanceData* InstanceData;
 	uint32 NumInstances;
 	int32 FreeBufferIndex;
 };
