#UE4 Nonunity fix

diff --git a/Engine/Source/Runtime/SlateRHIRenderer/Private/SlateUpdatableBuffer.h b/Engine/Source/Runtime/SlateRHIRenderer/Private/SlateUpdatableBuffer.h
--- a/Engine/Source/Runtime/SlateRHIRenderer/Private/SlateUpdatableBuffer.h
+++ b/Engine/Source/Runtime/SlateRHIRenderer/Private/SlateUpdatableBuffer.h
@@ -4,6 +4,7 @@
 
 #include "CoreMinimal.h"
 #include "Rendering/RenderingCommon.h"
+#include "RenderingThread.h"
 #include "Interfaces/ISlateRHIRendererModule.h"
 #include "SlateElementVertexBuffer.h"
 
