Removed TSlateElementVertexBuffer reference in RHI command FSlateUpdateInstanceBufferCommand as the instance could be destroyed before it is used in execution of the command
Rollbacked FDeferredCleanUpInterface use to ensure it wont crash for the moment

#rnx

diff --git a/Engine/Source/Runtime/SlateRHIRenderer/Private/SlateUpdatableBuffer.cpp b/Engine/Source/Runtime/SlateRHIRenderer/Private/SlateUpdatableBuffer.cpp
--- a/Engine/Source/Runtime/SlateRHIRenderer/Private/SlateUpdatableBuffer.cpp
+++ b/Engine/Source/Runtime/SlateRHIRenderer/Private/SlateUpdatableBuffer.cpp
@@ -7,14 +7,12 @@ DECLARE_CYCLE_STAT(TEXT("UpdateInstanceBuffer Time"), STAT_SlateUpdateInstanceBu
 
 struct FSlateUpdateInstanceBufferCommand final : public FRHICommand<FSlateUpdateInstanceBufferCommand>
 {
-	TSlateElementVertexBuffer<FVector4>& InstanceBuffer;
 	TArray<FVector4> InstanceData;
-	FVertexBufferRHIRef VertexBuffer;
+	FVertexBufferRHIRef VertexBufferRHI;
 
 	FSlateUpdateInstanceBufferCommand(TSlateElementVertexBuffer<FVector4>& InInstanceBuffer, const TArray<FVector4>& InInstanceData)
-		: InstanceBuffer(InInstanceBuffer)
-		, InstanceData(InInstanceData)
-		, VertexBuffer(InInstanceBuffer.VertexBufferRHI) // Keep a reference of RHIVertexBuffer in case it changes on RT thread during a resize
+		: InstanceData(InInstanceData)
+		, VertexBufferRHI(InInstanceBuffer.VertexBufferRHI)
 	{}
 
 	void Execute(FRHICommandListBase& CmdList)
@@ -22,39 +20,51 @@ struct FSlateUpdateInstanceBufferCommand final : public FRHICommand<FSlateUpdate
 		SCOPE_CYCLE_COUNTER( STAT_SlateUpdateInstanceBuffer );
 		const bool bIsInRenderingThread = !IsRunningRHIInSeparateThread() || CmdList.Bypass();
 
-		uint8* InstanceBufferData = (uint8*)InstanceBuffer.LockBuffer(InstanceData.Num(), bIsInRenderingThread);
+		uint32 RequiredBufferSize = InstanceData.Num() * sizeof(FVector4);
+		uint8* InstanceBufferData = nullptr;
 
-		FMemory::Memcpy( InstanceBufferData, InstanceData.GetData(), InstanceData.Num()*sizeof(FVector4) );
+		if (bIsInRenderingThread)
+		{
+			InstanceBufferData = static_cast<uint8*>(RHILockVertexBuffer(VertexBufferRHI, 0, RequiredBufferSize, RLM_WriteOnly));
+		}
+		else
+		{
+			InstanceBufferData = static_cast<uint8*>(GDynamicRHI->RHILockVertexBuffer(VertexBufferRHI, 0, RequiredBufferSize, RLM_WriteOnly));
+		}
+
+		FMemory::Memcpy(InstanceBufferData, InstanceData.GetData(), RequiredBufferSize);
 	
-		InstanceBuffer.UnlockBuffer(bIsInRenderingThread);
+		if (bIsInRenderingThread)
+		{
+			RHIUnlockVertexBuffer(VertexBufferRHI);
+		}
+		else
+		{
+			GDynamicRHI->RHIUnlockVertexBuffer(VertexBufferRHI);
+		}
 	}
 };
 
 FSlateUpdatableInstanceBuffer::FSlateUpdatableInstanceBuffer( int32 InitialInstanceCount )
 	: FreeBufferIndex(0)
 {
-	InstanceData = new FInstanceData;
-
-	InstanceData->InstanceBufferResource.Init(InitialInstanceCount);
+	InstanceBufferResource.Init(InitialInstanceCount);
 	for( int32 BufferIndex = 0; BufferIndex < SlateRHIConstants::NumBuffers; ++BufferIndex )
 	{
-		InstanceData->Array[BufferIndex].Reserve( InitialInstanceCount );
+		BufferData[BufferIndex].Reserve( InitialInstanceCount );
 	}
-}
 
-FSlateUpdatableInstanceBuffer::~FSlateUpdatableInstanceBuffer()
-{
-	InstanceData->InstanceBufferResource.Destroy();
-	BeginCleanup(InstanceData);
 }
 
-FSlateUpdatableInstanceBuffer::FInstanceData::~FInstanceData()
+FSlateUpdatableInstanceBuffer::~FSlateUpdatableInstanceBuffer()
 {
+	InstanceBufferResource.Destroy();
+	FlushRenderingCommands();
 }
 
 void FSlateUpdatableInstanceBuffer::BindStreamSource(FRHICommandListImmediate& RHICmdList, int32 StreamIndex, uint32 InstanceOffset)
 {
-	RHICmdList.SetStreamSource(StreamIndex, InstanceData->InstanceBufferResource.VertexBufferRHI, InstanceOffset*sizeof(FVector4));
+	RHICmdList.SetStreamSource(StreamIndex, InstanceBufferResource.VertexBufferRHI, InstanceOffset*sizeof(FVector4));
 }
 
 TSharedPtr<class FSlateInstanceBufferUpdate> FSlateUpdatableInstanceBuffer::BeginUpdate()
@@ -87,24 +97,24 @@ void FSlateUpdatableInstanceBuffer::UpdateRenderingData(int32 NumInstancesToUse)
 
 TArray<FVector4>& FSlateUpdatableInstanceBuffer::GetBufferData()
 {
-	return InstanceData->Array[FreeBufferIndex];
+	return BufferData[FreeBufferIndex];
 }
 
 void FSlateUpdatableInstanceBuffer::UpdateRenderingData_RenderThread(FRHICommandListImmediate& RHICmdList, int32 BufferIndex)
 {
 	SCOPE_CYCLE_COUNTER( STAT_SlateUpdateInstanceBuffer );
 
-	const TArray<FVector4>& RenderThreadBufferData = InstanceData->Array[BufferIndex];
-	InstanceData->InstanceBufferResource.PreFillBuffer( RenderThreadBufferData.Num(), false );
+	const TArray<FVector4>& RenderThreadBufferData = BufferData[BufferIndex];
+	InstanceBufferResource.PreFillBuffer( RenderThreadBufferData.Num(), false );
 
 	if(!IsRunningRHIInSeparateThread() || RHICmdList.Bypass())
 	{
-		FSlateUpdateInstanceBufferCommand Command(InstanceData->InstanceBufferResource, RenderThreadBufferData);
+		FSlateUpdateInstanceBufferCommand Command(InstanceBufferResource, RenderThreadBufferData);
 		Command.Execute(RHICmdList);
 	}
 	else
 	{
-		ALLOC_COMMAND_CL(RHICmdList, FSlateUpdateInstanceBufferCommand)(InstanceData->InstanceBufferResource, RenderThreadBufferData);
+		ALLOC_COMMAND_CL(RHICmdList, FSlateUpdateInstanceBufferCommand)(InstanceBufferResource, RenderThreadBufferData);
 	}
 }
 
diff --git a/Engine/Source/Runtime/SlateRHIRenderer/Private/SlateUpdatableBuffer.h b/Engine/Source/Runtime/SlateRHIRenderer/Private/SlateUpdatableBuffer.h
--- a/Engine/Source/Runtime/SlateRHIRenderer/Private/SlateUpdatableBuffer.h
+++ b/Engine/Source/Runtime/SlateRHIRenderer/Private/SlateUpdatableBuffer.h
@@ -30,20 +30,8 @@ private:
 	void UpdateRenderingData_RenderThread(FRHICommandListImmediate& RHICmdList, int32 BufferIndex);
 
 private:
-	class FInstanceData : public FDeferredCleanupInterface
-	{
-	public:
-		FInstanceData()
-		{
-		}
-
-		virtual ~FInstanceData();
-
-		TArray<FVector4> Array[SlateRHIConstants::NumBuffers];
-		TSlateElementVertexBuffer<FVector4> InstanceBufferResource;
-	};
-
-	FInstanceData* InstanceData;
+	TArray<FVector4> BufferData[SlateRHIConstants::NumBuffers];
+	TSlateElementVertexBuffer<FVector4> InstanceBufferResource;
 	uint32 NumInstances;
 	int32 FreeBufferIndex;
 };
