Add a CSV stat for slate renderthread time

diff --git a/Engine/Source/Runtime/SlateRHIRenderer/Private/SlateRHIRenderer.cpp b/Engine/Source/Runtime/SlateRHIRenderer/Private/SlateRHIRenderer.cpp
--- a/Engine/Source/Runtime/SlateRHIRenderer/Private/SlateRHIRenderer.cpp
+++ b/Engine/Source/Runtime/SlateRHIRenderer/Private/SlateRHIRenderer.cpp
@@ -702,6 +702,7 @@ void FSlateRHIRenderer::DrawWindow_RenderThread(FRHICommandListImmediate& RHICmd
 		{
 			SCOPED_GPU_STAT(RHICmdList, SlateUI);
 			SCOPE_CYCLE_COUNTER(STAT_SlateRenderingRTTime);
+			CSV_SCOPED_TIMING_STAT_EXCLUSIVE(Slate);
 
 			FSlateBatchData& BatchData = WindowElementList.GetBatchData();
 
