Type safe FRenderThreadUpdateContext

diff --git a/Engine/Source/Runtime/SlateCore/Public/Rendering/SlateRenderer.h b/Engine/Source/Runtime/SlateCore/Public/Rendering/SlateRenderer.h
--- a/Engine/Source/Runtime/SlateCore/Public/Rendering/SlateRenderer.h
+++ b/Engine/Source/Runtime/SlateCore/Public/Rendering/SlateRenderer.h
@@ -9,8 +9,10 @@
 #include "Templates/RefCounting.h"
 
 class FRHITexture2D;
+class FRenderTarget;
 class FSlateDrawBuffer;
 class FSlateUpdatableTexture;
+class ISlate3DRenderer;
 class ISlateAtlasProvider;
 class ISlateStyle;
 class SWindow;
@@ -29,8 +31,8 @@ struct FRenderThreadUpdateContext
 	float WorldTimeSeconds;
 	float DeltaTimeSeconds;
 	float RealTimeSeconds;
-	void* RenderTargetResource;
-	void* Renderer;
+	FRenderTarget* RenderTarget;
+	ISlate3DRenderer* Renderer;
 	bool bClearTarget;
 };
 
diff --git a/Engine/Source/Runtime/SlateRHIRenderer/Private/Slate3DRenderer.cpp b/Engine/Source/Runtime/SlateRHIRenderer/Private/Slate3DRenderer.cpp
--- a/Engine/Source/Runtime/SlateRHIRenderer/Private/Slate3DRenderer.cpp
+++ b/Engine/Source/Runtime/SlateRHIRenderer/Private/Slate3DRenderer.cpp
@@ -124,16 +124,15 @@ void FSlate3DRenderer::DrawWindowToTarget_RenderThread(FRHICommandListImmediate&
 	SCOPED_DRAW_EVENT( InRHICmdList, SlateRenderToTarget );
 	SCOPED_GPU_STAT(InRHICmdList, Slate3D);
 
-	checkSlow( Context.RenderTargetResource );
+	checkSlow(Context.RenderTarget);
 
 	const TArray<TSharedRef<FSlateWindowElementList>>& WindowsToDraw = Context.WindowDrawBuffer->GetWindowElementLists();
 
 	// Enqueue a command to unlock the draw buffer after all windows have been drawn
 	RenderTargetPolicy->BeginDrawingWindows();
 
-	FRenderTarget* RenderTargetResource = static_cast<FRenderTarget*>(Context.RenderTargetResource);
 	// Set render target and clear.
-	FTexture2DRHIRef RTTextureRHI = RenderTargetResource->GetRenderTargetTexture();
+	FTexture2DRHIRef RTTextureRHI = Context.RenderTarget->GetRenderTargetTexture();
 	InRHICmdList.TransitionResource(EResourceTransitionAccess::EWritable, RTTextureRHI);
 	
 	FRHIRenderPassInfo RPInfo(RTTextureRHI, ERenderTargetActions::Load_Store);
@@ -160,14 +159,14 @@ void FSlate3DRenderer::DrawWindowToTarget_RenderThread(FRHICommandListImmediate&
 				FMatrix ViewOffset = FTranslationMatrix::Make(FVector(DrawOffset, 0));
 				ProjectionMatrix = ViewOffset * ProjectionMatrix;
 
-				FSlateBackBuffer BackBufferTarget(RenderTargetResource->GetRenderTargetTexture(), FIntPoint(RTTextureRHI->GetSizeX(), RTTextureRHI->GetSizeY()));
+				FSlateBackBuffer BackBufferTarget(Context.RenderTarget->GetRenderTargetTexture(), FIntPoint(RTTextureRHI->GetSizeX(), RTTextureRHI->GetSizeY()));
 
 				FSlateRenderingParams DrawOptions(ProjectionMatrix, Context.WorldTimeSeconds, Context.DeltaTimeSeconds, Context.RealTimeSeconds);
 				// The scene renderer will handle it in this case
 				DrawOptions.bAllowSwitchVerticalAxis = false;
 				DrawOptions.ViewOffset = DrawOffset;
 
-				FTexture2DRHIRef ColorTarget = RenderTargetResource->GetRenderTargetTexture();
+				FTexture2DRHIRef ColorTarget = Context.RenderTarget->GetRenderTargetTexture();
 
 				if (BatchData.IsStencilClippingRequired())
 				{
@@ -200,7 +199,7 @@ void FSlate3DRenderer::DrawWindowToTarget_RenderThread(FRHICommandListImmediate&
 	}
 
 	FSlateEndDrawingWindowsCommand::EndDrawingWindows(InRHICmdList, Context.WindowDrawBuffer, *RenderTargetPolicy);
-	InRHICmdList.CopyToResolveTarget(RenderTargetResource->GetRenderTargetTexture(), RTTextureRHI, FResolveParams());
+	InRHICmdList.CopyToResolveTarget(Context.RenderTarget->GetRenderTargetTexture(), RTTextureRHI, FResolveParams());
 
 	ISlate3DRendererPtr Self = SharedThis(this);
 
diff --git a/Engine/Source/Runtime/SlateRHIRenderer/Private/SlateRHIRenderer.cpp b/Engine/Source/Runtime/SlateRHIRenderer/Private/SlateRHIRenderer.cpp
--- a/Engine/Source/Runtime/SlateRHIRenderer/Private/SlateRHIRenderer.cpp
+++ b/Engine/Source/Runtime/SlateRHIRenderer/Private/SlateRHIRenderer.cpp
@@ -1189,7 +1189,7 @@ void FSlateRHIRenderer::DrawWindows_Private(FSlateDrawBuffer& WindowDrawBuffer)
 			{
 				for (const FRenderThreadUpdateContext& Context : Contexts)
 				{
-					static_cast<ISlate3DRenderer*>(Context.Renderer)->DrawWindowToTarget_RenderThread(RHICmdList, Context);
+					Context.Renderer->DrawWindowToTarget_RenderThread(RHICmdList, Context);
 				}
 			}
 		);
@@ -1631,7 +1631,7 @@ void FSlateRHIRenderer::AddWidgetRendererUpdate(const struct FRenderThreadUpdate
 		ENQUEUE_RENDER_COMMAND(DrawWidgetRendererImmediate)(
 			[InContext](FRHICommandListImmediate& RHICmdList)
 			{
-				static_cast<ISlate3DRenderer*>(InContext.Renderer)->DrawWindowToTarget_RenderThread(RHICmdList, InContext);
+				InContext.Renderer->DrawWindowToTarget_RenderThread(RHICmdList, InContext);
 			});
 	}
 }
