Defer cached slate rendering data clean up until the end of the frame to ensure it doesnt get out of order with deferred updates from retainers

diff --git a/Engine/Source/Runtime/SlateCore/Private/Rendering/DrawElements.cpp b/Engine/Source/Runtime/SlateCore/Private/Rendering/DrawElements.cpp
--- a/Engine/Source/Runtime/SlateCore/Private/Rendering/DrawElements.cpp
+++ b/Engine/Source/Runtime/SlateCore/Private/Rendering/DrawElements.cpp
@@ -954,7 +954,7 @@ FSlateDrawElement& FSlateCachedElementData::AddCachedElement(FSlateCachedElement
 		// We need to cache this clip state for the next time the element draws
 		FSlateCachedClipState& CachedClipState = FindOrAddCachedClipState(ExistingClipState);
 		List.AddCachedClipState(CachedClipState);
-		NewElement.SetCachedClippingState(&CachedClipState.ClippingState);
+		NewElement.SetCachedClippingState(&CachedClipState.ClippingState.Get());
 	}
 
 	return NewElement;
@@ -972,22 +972,21 @@ FSlateCachedClipState& FSlateCachedElementData::FindOrAddCachedClipState(const F
 {
 	for (auto& CachedState : CachedClipStates)
 	{
-		if (CachedState->ClippingState == *RefClipState)
+		if (*CachedState.ClippingState == *RefClipState)
 		{
-			CachedState->BeginUsingState();
-			return *CachedState;
+			return CachedState;
 		}
 	}
 
-	return *CachedClipStates.Emplace_GetRef(MakeUnique<FSlateCachedClipState>(*RefClipState));
+	return CachedClipStates.Emplace_GetRef(FSlateCachedClipState(*RefClipState));
 }
 
 void FSlateCachedElementData::CleanupUnusedClipStates()
 {
 	for(int32 CachedStateIdx = 0; CachedStateIdx < CachedClipStates.Num();)
 	{
-		const TUniquePtr<FSlateCachedClipState>& CachedState = CachedClipStates[CachedStateIdx];
-		if (CachedState->GetUsageCount() == 0)
+		const FSlateCachedClipState& CachedState = CachedClipStates[CachedStateIdx];
+		if (CachedState.ClippingState.IsUnique())
 		{
 			CachedClipStates.RemoveAtSwap(CachedStateIdx);
 		}
@@ -1033,7 +1032,7 @@ FSlateRenderBatch& FSlateCachedElementList::AddRenderBatch(int32 InLayer, const
 
 void FSlateCachedElementList::AddCachedClipState(FSlateCachedClipState& ClipStateToCache)
 {
-	CachedRenderingData->CachedClipStates.Add(&ClipStateToCache);
+	CachedRenderingData->CachedClipStates.Add(ClipStateToCache);
 }
 
 void FSlateCachedElementList::AddReferencedObjects(FReferenceCollector& Collector)
diff --git a/Engine/Source/Runtime/SlateCore/Public/Layout/Clipping.h b/Engine/Source/Runtime/SlateCore/Public/Layout/Clipping.h
--- a/Engine/Source/Runtime/SlateCore/Public/Layout/Clipping.h
+++ b/Engine/Source/Runtime/SlateCore/Public/Layout/Clipping.h
@@ -331,25 +331,10 @@ class FSlateCachedClipState
 {
 public:
 	FSlateCachedClipState(const FSlateClippingState& InState)
-		: ClippingState(InState)
-		, UsageCount(1)
+		: ClippingState(MakeShared<FSlateClippingState, ESPMode::ThreadSafe>(InState))
 	{}
 
-	FSlateClippingState ClippingState;
-
-	int32 GetUsageCount() const { return UsageCount; }
-	void BeginUsingState()
-	{
-		FPlatformAtomics::InterlockedIncrement(&UsageCount);
-	}
-
-	void EndUsingState()
-	{
-		int32 NewCount = FPlatformAtomics::InterlockedDecrement(&UsageCount);
-		check(NewCount >= 0);
-	}
-private:
-	int32 UsageCount;
+	TSharedRef<FSlateClippingState, ESPMode::ThreadSafe> ClippingState;
 };
 
 
diff --git a/Engine/Source/Runtime/SlateCore/Public/Rendering/DrawElements.h b/Engine/Source/Runtime/SlateCore/Public/Rendering/DrawElements.h
--- a/Engine/Source/Runtime/SlateCore/Public/Rendering/DrawElements.h
+++ b/Engine/Source/Runtime/SlateCore/Public/Rendering/DrawElements.h
@@ -339,13 +339,10 @@ struct FSlateCachedFastPathRenderingData
 {
 	~FSlateCachedFastPathRenderingData()
 	{
-		for (FSlateCachedClipState* State : CachedClipStates)
-		{
-			State->EndUsingState();
-		}
+		CachedClipStates.Reset();
 	}
 
-	TArray<class FSlateCachedClipState*, TInlineAllocator<1>> CachedClipStates;
+	TArray<FSlateCachedClipState, TInlineAllocator<1>> CachedClipStates;
 	FSlateVertexArray Vertices;
 	FSlateIndexArray Indices;
 };
@@ -427,7 +424,7 @@ public:
 	TDoubleLinkedList<FSlateCachedElementList> CachedElementLists;
 private:
 
-	TArray<TUniquePtr<FSlateCachedClipState>> CachedClipStates;
+	TArray<FSlateCachedClipState> CachedClipStates;
 };
 
 /**
diff --git a/Engine/Source/Runtime/SlateRHIRenderer/Private/SlateRHIRenderer.cpp b/Engine/Source/Runtime/SlateRHIRenderer/Private/SlateRHIRenderer.cpp
--- a/Engine/Source/Runtime/SlateRHIRenderer/Private/SlateRHIRenderer.cpp
+++ b/Engine/Source/Runtime/SlateRHIRenderer/Private/SlateRHIRenderer.cpp
@@ -153,6 +153,7 @@ FSlateRHIRenderer::FSlateRHIRenderer(TSharedRef<FSlateFontServices> InSlateFontS
 	: FSlateRenderer(InSlateFontServices)
 	, EnqueuedWindowDrawBuffer(NULL)
 	, FreeBufferIndex(0)
+	, FastPathRenderingDataCleanupList(nullptr)
 	, CurrentSceneIndex(-1)
 	, ResourceVersion(0)
 {
@@ -214,10 +215,14 @@ void FSlateRHIRenderer::Destroy()
 		BeginReleaseResource(It.Value());
 	}
 
+	if (FastPathRenderingDataCleanupList)
+	{
+		FastPathRenderingDataCleanupList->Cleanup();
+		FastPathRenderingDataCleanupList = nullptr;
+	}
 
 	FlushRenderingCommands();
 
-
 	ElementBatcher.Reset();
 	RenderingPolicy.Reset();
 	ResourceManager.Reset();
@@ -1197,6 +1202,12 @@ void FSlateRHIRenderer::DrawWindows_Private(FSlateDrawBuffer& WindowDrawBuffer)
 		DeferredUpdateContexts.Empty();
 	}
 
+	if (FastPathRenderingDataCleanupList)
+	{
+		FastPathRenderingDataCleanupList->Cleanup();
+		FastPathRenderingDataCleanupList = nullptr;
+	}
+
 	// flush the cache if needed
 	FontCache->ConditionalFlushCache();
 }
@@ -1478,24 +1489,6 @@ void FSlateRHIRenderer::ClearScenes()
 }
 
 
-struct FClearCachedRenderingDataCommand final : public FRHICommand < FClearCachedRenderingDataCommand >
-{
-public:
-	FClearCachedRenderingDataCommand(FSlateCachedFastPathRenderingData* InCachedRenderingData)
-		: CachedRenderingData(InCachedRenderingData)
-	{
-		
-	}
-
-	void Execute(FRHICommandListBase& CmdList)
-	{
-		delete CachedRenderingData;
-	}
-
-private:
-	FSlateCachedFastPathRenderingData* CachedRenderingData;
-};
-
 struct FClearCachedElementDataCommand final : public FRHICommand < FClearCachedElementDataCommand >
 {
 public:
@@ -1518,21 +1511,13 @@ void FSlateRHIRenderer::DestroyCachedFastPathRenderingData(FSlateCachedFastPathR
 {
 	check(CachedRenderingData);
 
-	// Cached data should be destroyed in a thread safe way.  If there is an rhi thread it could be reading from the data to copy it into a vertex buffer
-	// so delete it on the rhi thread if necessary, otherwise delete it on the render thread
-	ENQUEUE_RENDER_COMMAND(ClearCachedRenderingData)(
-		[CachedRenderingData](FRHICommandListImmediate& RHICmdList)
+	if (!FastPathRenderingDataCleanupList)
 	{
-		if (!RHICmdList.Bypass())
-		{
-			new (RHICmdList.AllocCommand<FClearCachedRenderingDataCommand>()) FClearCachedRenderingDataCommand(CachedRenderingData);
-		}
-		else
-		{
-			FClearCachedRenderingDataCommand Cmd(CachedRenderingData);
-			Cmd.Execute(RHICmdList);
-		}
-	});
+		// This will be deleted later on the rendering thread
+		FastPathRenderingDataCleanupList = new FFastPathRenderingDataCleanupList;
+	}
+
+	FastPathRenderingDataCleanupList->FastPathRenderingDataToRemove.Add(CachedRenderingData);
 }
 
 void FSlateRHIRenderer::DestroyCachedFastPathElementData(FSlateCachedElementData* CachedElementData)
@@ -1659,3 +1644,50 @@ void FSlateEndDrawingWindowsCommand::EndDrawingWindows(FRHICommandListImmediate&
 		Cmd.Execute(RHICmdList);
 	}
 }
+
+
+struct FClearCachedRenderingDataCommand final : public FRHICommand < FClearCachedRenderingDataCommand >
+{
+public:
+	FClearCachedRenderingDataCommand(FFastPathRenderingDataCleanupList* InCleanupList)
+		: CleanupList(InCleanupList)
+	{
+
+	}
+
+	void Execute(FRHICommandListBase& CmdList)
+	{
+		delete CleanupList;
+	}
+
+private:
+	FFastPathRenderingDataCleanupList* CleanupList;
+};
+
+
+FFastPathRenderingDataCleanupList::~FFastPathRenderingDataCleanupList()
+{
+	for (FSlateCachedFastPathRenderingData* Data : FastPathRenderingDataToRemove)
+	{
+		delete Data;
+	}
+}
+
+void FFastPathRenderingDataCleanupList::Cleanup()
+{
+	// Cached data should be destroyed in a thread safe way.  If there is an rhi thread it could be reading from the data to copy it into a vertex buffer
+	// so delete it on the rhi thread if necessary, otherwise delete it on the render thread
+	ENQUEUE_RENDER_COMMAND(ClearCachedRenderingData)(
+		[CleanupList = this](FRHICommandListImmediate& RHICmdList)
+	{
+		if (!RHICmdList.Bypass())
+		{
+			new (RHICmdList.AllocCommand<FClearCachedRenderingDataCommand>()) FClearCachedRenderingDataCommand(CleanupList);
+		}
+		else
+		{
+			FClearCachedRenderingDataCommand Cmd(CleanupList);
+			Cmd.Execute(RHICmdList);
+		}
+	});
+}
\ No newline at end of file
diff --git a/Engine/Source/Runtime/SlateRHIRenderer/Private/SlateRHIRenderer.h b/Engine/Source/Runtime/SlateRHIRenderer/Private/SlateRHIRenderer.h
--- a/Engine/Source/Runtime/SlateRHIRenderer/Private/SlateRHIRenderer.h
+++ b/Engine/Source/Runtime/SlateRHIRenderer/Private/SlateRHIRenderer.h
@@ -144,6 +144,14 @@ struct FViewportInfo : public FRenderResource
 	}	
 };
 
+struct FFastPathRenderingDataCleanupList
+{
+	TArray<FSlateCachedFastPathRenderingData*, TInlineAllocator<20>> FastPathRenderingDataToRemove;
+
+	~FFastPathRenderingDataCleanupList();
+	void Cleanup();
+};
+
 /** A Slate rendering implementation for Unreal engine */
 class FSlateRHIRenderer : public FSlateRenderer
 {
@@ -273,7 +281,9 @@ private:
 	/** Drawing policy */
 	TSharedPtr<FSlateRHIRenderingPolicy> RenderingPolicy;
 
-	TArray< TSharedPtr<FSlateDynamicImageBrush> > DynamicBrushesToRemove[NumDrawBuffers];
+	TArray<TSharedPtr<FSlateDynamicImageBrush>> DynamicBrushesToRemove[NumDrawBuffers];
+
+	FFastPathRenderingDataCleanupList* FastPathRenderingDataCleanupList;
 
 	FDeferredUpdateContextList DeferredUpdateContexts;
 
