Add support to the Editor to allow rendering onto HDR monitors

diff --git a/Engine/Shaders/Private/CompositeUIPixelShader.usf b/Engine/Shaders/Private/CompositeUIPixelShader.usf
--- a/Engine/Shaders/Private/CompositeUIPixelShader.usf
+++ b/Engine/Shaders/Private/CompositeUIPixelShader.usf
@@ -72,6 +72,28 @@ float3 ColorLookupTable(float3 LinearColor)
 	return ColorSpaceLUT.SampleLevel(ColorSpaceLUTSampler, UVW, 0).rgb * LUT_SCALING;
 }
 
+// Given an input using Rec709 primaries and a sRGB tranfer function where
+// 1.0 maps to paperwhitenits defined in editor. Remap to Rec2020 with PQ transfer function.
+void HDREditorConvert(
+	FScreenVertexOutput Input,
+	out float4 OutColor : SV_Target0
+	)
+{
+	OutColor = 0;
+	float3 SceneColor = Texture2DSample(SceneTexture, SceneSampler, Input.UV).xyz ;
+
+	// Get color into Rec2020 space
+	SceneColor = sRGBToLinear( SceneColor) ;
+	SceneColor = SceneColor * UILevel ; 
+
+	SceneColor = mul(sRGB_2_XYZ_MAT, SceneColor);
+	SceneColor = mul(XYZ_2_Rec2020_MAT, SceneColor);
+
+	// Do PQ Curve
+	SceneColor = LinearToST2084( SceneColor ); 
+	OutColor.rgb = SceneColor;
+}
+
 void Main(
 	FScreenVertexOutput Input,
 	out float4 OutColor : SV_Target0
diff --git a/Engine/Shaders/Private/PostProcessTonemap.usf b/Engine/Shaders/Private/PostProcessTonemap.usf
--- a/Engine/Shaders/Private/PostProcessTonemap.usf
+++ b/Engine/Shaders/Private/PostProcessTonemap.usf
@@ -51,6 +51,7 @@ float4 SceneUVMinMax;
 float4 SceneBloomUVMinMax;
 
 float4 ScreenPosToScenePixel;
+float4 HDRParams;
 
 #if EYEADAPTATION_EXPOSURE_FIX != 1
 // Default value used when eye adaptation is disabled (e.g fixed exposure) or not available.
@@ -580,6 +581,16 @@ float4 TonemapCommonPS(
 	}
 #endif //DEBUG_TONEMAPPER
 
+#if SM5_PROFILE && !USE_GAMMA_ONLY
+	// If HDR in the editor need to convert from PQ space to linear sRGB
+	BRANCH
+	if(HDRParams.x == 1.0f)
+	{
+		OutColor.rgb = ST2084ToLinear(OutColor.rgb);
+		OutColor.rgb = OutColor.rgb / HDRParams.y;
+		OutColor.rgb = LinearToSrgbBranchless( OutColor.rgb );
+	}
+#endif
 	return OutColor;
 }
 
diff --git a/Engine/Source/Editor/UnrealEd/Classes/Settings/EditorExperimentalSettings.h b/Engine/Source/Editor/UnrealEd/Classes/Settings/EditorExperimentalSettings.h
--- a/Engine/Source/Editor/UnrealEd/Classes/Settings/EditorExperimentalSettings.h
+++ b/Engine/Source/Editor/UnrealEd/Classes/Settings/EditorExperimentalSettings.h
@@ -20,6 +20,14 @@ class UNREALED_API UEditorExperimentalSettings
 
 public:
 
+	/** Allows the editor to run on HDR monitors on Windows 10 */
+	UPROPERTY(EditAnywhere, config, Category = HDR, meta = (ConfigRestartRequired = true, DisplayName = "Enable Editor Support for HDR Monitors"))
+	bool bHDREditor;
+
+	/** The brightness of the slate UI on HDR monitors */
+	UPROPERTY(EditAnywhere, config, Category = HDR, meta = (ClampMin = "100.0", ClampMax = "300.0", UIMin = "100.0", UIMax = "300.0"))
+	float HDREditorNITLevel;
+
 	/** Allows usage of the procedural foliage system */
 	UPROPERTY(EditAnywhere, config, Category = Foliage, meta = (DisplayName = "Procedural Foliage"))
 	bool bProceduralFoliage;
@@ -134,8 +142,8 @@ public:
 protected:
 
 	// UObject overrides
-
-	virtual void PostEditChangeProperty( struct FPropertyChangedEvent& PropertyChangedEvent ) override;
+	virtual void PostEditChangeProperty(struct FPropertyChangedEvent& PropertyChangedEvent) override;
+	virtual void PostInitProperties() override;
 
 private:
 
diff --git a/Engine/Source/Editor/UnrealEd/Private/EditorEngine.cpp b/Engine/Source/Editor/UnrealEd/Private/EditorEngine.cpp
--- a/Engine/Source/Editor/UnrealEd/Private/EditorEngine.cpp
+++ b/Engine/Source/Editor/UnrealEd/Private/EditorEngine.cpp
@@ -1996,6 +1996,14 @@ bool UEditorEngine::UpdateSingleViewportClient(FEditorViewportClient* InViewport
 {
 	bool bUpdatedNonRealtimeViewport = false;
 
+	// When rendering the viewport we need to know whether the final result will be shown on a HDR display. This affects the final post processing step
+	FSceneViewport *SceneViewport = static_cast<FSceneViewport*>(InViewportClient->Viewport);
+	TSharedPtr<SWindow> Window = SceneViewport->FindWindow();
+	if (Window)
+	{
+		InViewportClient->Viewport->SetHDRMode(Window->GetIsHDR());
+	}
+
 	// Always submit view information for content streaming 
 	// otherwise content for editor view can be streamed out if there are other views (ex: thumbnails)
 	if (InViewportClient->IsPerspective())
diff --git a/Engine/Source/Editor/UnrealEd/Private/EditorViewportClient.cpp b/Engine/Source/Editor/UnrealEd/Private/EditorViewportClient.cpp
--- a/Engine/Source/Editor/UnrealEd/Private/EditorViewportClient.cpp
+++ b/Engine/Source/Editor/UnrealEd/Private/EditorViewportClient.cpp
@@ -3652,6 +3652,8 @@ void FEditorViewportClient::Draw(FViewport* InViewport, FCanvas* Canvas)
 
 	ViewFamily.EngineShowFlags = EngineShowFlags;
 
+	ViewFamily.bIsHDR = Viewport->IsHDRViewport();
+
 	UpdateDebugViewModeShaders();
 
 	if( ModeTools->GetActiveMode( FBuiltinEditorModes::EM_InterpEdit ) == 0 || !AllowsCinematicControl() )
diff --git a/Engine/Source/Editor/UnrealEd/Private/Settings/SettingsClasses.cpp b/Engine/Source/Editor/UnrealEd/Private/Settings/SettingsClasses.cpp
--- a/Engine/Source/Editor/UnrealEd/Private/Settings/SettingsClasses.cpp
+++ b/Engine/Source/Editor/UnrealEd/Private/Settings/SettingsClasses.cpp
@@ -131,14 +131,35 @@ USkeletalMeshEditorSettings::USkeletalMeshEditorSettings(const FObjectInitialize
 /* UEditorExperimentalSettings interface
  *****************************************************************************/
 
+static TAutoConsoleVariable<int32> CVarEditorHDRSupport(
+	TEXT("Editor.HDRSupport"),
+	0,
+	TEXT("Sets whether or not we should allow the editor to run on HDR monitors"),
+	ECVF_Default);
+
+static TAutoConsoleVariable<float> CVarEditorHDRNITLevel(
+	TEXT("Editor.HDRNITLevel"),
+	160.0f,
+	TEXT("Sets The desired NIT level of the editor when running on HDR"),
+	ECVF_Default);
+
 UEditorExperimentalSettings::UEditorExperimentalSettings( const FObjectInitializer& ObjectInitializer )
 	: Super(ObjectInitializer)
+	, bHDREditor(false)
+	, HDREditorNITLevel(160.0f)
 	, bEnableLocalizationDashboard(true)
 	, bUseOpenCLForConvexHullDecomp(false)
 	, bAllowPotentiallyUnsafePropertyEditing(false)
 {
 }
 
+void UEditorExperimentalSettings::PostInitProperties()
+{
+	CVarEditorHDRSupport->Set(bHDREditor ? 1 : 0, ECVF_SetByProjectSetting);
+	CVarEditorHDRNITLevel->Set(HDREditorNITLevel, ECVF_SetByProjectSetting);
+	Super::PostInitProperties();
+}
+
 void UEditorExperimentalSettings::PostEditChangeProperty( struct FPropertyChangedEvent& PropertyChangedEvent )
 {
 	static const FName NAME_EQS = GET_MEMBER_NAME_CHECKED(UEditorExperimentalSettings, bEQSEditor);
@@ -158,7 +179,14 @@ void UEditorExperimentalSettings::PostEditChangeProperty( struct FPropertyChange
 			FModuleManager::Get().LoadModule(TEXT("EnvironmentQueryEditor"));
 		}
 	}
-
+	else if (Name == FName(TEXT("bHDREditor")))
+	{
+		CVarEditorHDRSupport->Set(bHDREditor ? 1 : 0, ECVF_SetByProjectSetting);
+	}
+	else if (Name == FName(TEXT("HDREditorNITLevel")))
+	{
+		CVarEditorHDRNITLevel->Set(HDREditorNITLevel, ECVF_SetByProjectSetting);
+	}
 	if (!FUnrealEdMisc::Get().IsDeletePreferences())
 	{
 		SaveConfig();
diff --git a/Engine/Source/Runtime/Engine/Private/GameViewportClient.cpp b/Engine/Source/Runtime/Engine/Private/GameViewportClient.cpp
--- a/Engine/Source/Runtime/Engine/Private/GameViewportClient.cpp
+++ b/Engine/Source/Runtime/Engine/Private/GameViewportClient.cpp
@@ -1505,6 +1505,9 @@ void UGameViewportClient::Draw(FViewport* InViewport, FCanvas* SceneCanvas)
 		}
 	}
 
+	
+	ViewFamily.bIsHDR = GetWindow().IsValid() ? GetWindow().Get()->GetIsHDR() : false;
+
 	// Draw the player views.
 	if (!bDisableWorldRendering && PlayerViewMap.Num() > 0 && FSlateApplication::Get().GetPlatformApplication()->IsAllowedToRender()) //-V560
 	{
diff --git a/Engine/Source/Runtime/Engine/Private/SceneView.cpp b/Engine/Source/Runtime/Engine/Private/SceneView.cpp
--- a/Engine/Source/Runtime/Engine/Private/SceneView.cpp
+++ b/Engine/Source/Runtime/Engine/Private/SceneView.cpp
@@ -2378,6 +2378,7 @@ FSceneViewFamily::FSceneViewFamily(const ConstructionValues& CVS)
 	SceneCaptureSource(SCS_FinalColorLDR),
 	SceneCaptureCompositeMode(SCCM_Overwrite),
 	bWorldIsPaused(false),
+	bIsHDR(false),
 	GammaCorrection(CVS.GammaCorrection),
 	SecondaryViewFraction(1.0f),
 	SecondaryScreenPercentageMethod(ESecondaryScreenPercentageMethod::LowerPixelDensitySimulation),
diff --git a/Engine/Source/Runtime/Engine/Private/Slate/SceneViewport.cpp b/Engine/Source/Runtime/Engine/Private/Slate/SceneViewport.cpp
--- a/Engine/Source/Runtime/Engine/Private/Slate/SceneViewport.cpp
+++ b/Engine/Source/Runtime/Engine/Private/Slate/SceneViewport.cpp
@@ -1943,6 +1943,16 @@ void FSceneViewport::InitDynamicRHI()
 
 		static const auto CVarDefaultBackBufferPixelFormat = IConsoleManager::Get().FindTConsoleVariableDataInt(TEXT("r.DefaultBackBufferPixelFormat"));
 		EPixelFormat SceneTargetFormat = EDefaultBackBufferPixelFormat::Convert2PixelFormat(EDefaultBackBufferPixelFormat::FromInt(CVarDefaultBackBufferPixelFormat->GetValueOnRenderThread()));
+	
+#if WITH_EDITOR
+		// HDR Editor needs to be in float format if running with HDR
+		static auto CVarHDREnable = IConsoleManager::Get().FindConsoleVariable(TEXT("Editor.HDRSupport"));
+		if(CVarHDREnable && (CVarHDREnable->GetInt() != 0))
+		{
+			SceneTargetFormat = PF_FloatRGBA;
+		}
+#endif
+
 #if PLATFORM_HTML5
 		if ( SceneTargetFormat == PF_A2B10G10R10 )
 		{	// UE-71220: this seems to be FDummyViewport - force to valid pixel format type for HTML5
diff --git a/Engine/Source/Runtime/Engine/Private/UnrealClient.cpp b/Engine/Source/Runtime/Engine/Private/UnrealClient.cpp
--- a/Engine/Source/Runtime/Engine/Private/UnrealClient.cpp
+++ b/Engine/Source/Runtime/Engine/Private/UnrealClient.cpp
@@ -1166,6 +1166,7 @@ FViewport::FViewport(FViewportClient* InViewportClient):
 	bHitProxiesCached(false),
 	bHasRequestedToggleFreeze(false),
 	bIsSlateViewport(false),
+	bIsHDR(false),
 	bTakeHighResScreenShot(false)
 {
 	//initialize the hit proxy kernel
diff --git a/Engine/Source/Runtime/Engine/Public/SceneView.h b/Engine/Source/Runtime/Engine/Public/SceneView.h
--- a/Engine/Source/Runtime/Engine/Public/SceneView.h
+++ b/Engine/Source/Runtime/Engine/Public/SceneView.h
@@ -1508,6 +1508,9 @@ public:
 	 */
 	bool bWorldIsPaused;
 
+	/** When enabled, the post processing will output in HDR space */
+	bool bIsHDR;
+
 	/** Gamma correction used when rendering this family. Default is 1.0 */
 	float GammaCorrection;
 	
diff --git a/Engine/Source/Runtime/Engine/Public/UnrealClient.h b/Engine/Source/Runtime/Engine/Public/UnrealClient.h
--- a/Engine/Source/Runtime/Engine/Public/UnrealClient.h
+++ b/Engine/Source/Runtime/Engine/Public/UnrealClient.h
@@ -575,6 +575,12 @@ public:
 	/** Returns true if this is an FSlateSceneViewport */
 	bool IsSlateViewport() const { return bIsSlateViewport; }
 
+	/** Returns true if this viewport should be rendered in HDR */
+	bool IsHDRViewport() const { return bIsHDR; }
+
+	/** Sets HDR Status of Viewport */
+	void SetHDRMode( bool bHDR) { bIsHDR = bHDR; }
+
 	/** The current version of the running instance */
 	FString AppVersionString;
 
@@ -697,6 +703,9 @@ protected:
 	/** If true this viewport is an FSlateSceneViewport */
 	uint32 bIsSlateViewport : 1;
 
+	/** If true this viewport is being displayed on a HDR monitor */
+	uint32 bIsHDR : 1;
+
 	/** true if we should draw game viewports (has no effect on Editor viewports) */
 	ENGINE_API static bool bIsGameRenderingEnabled;
 
diff --git a/Engine/Source/Runtime/NullDrv/Public/NullRHI.h b/Engine/Source/Runtime/NullDrv/Public/NullRHI.h
--- a/Engine/Source/Runtime/NullDrv/Public/NullRHI.h
+++ b/Engine/Source/Runtime/NullDrv/Public/NullRHI.h
@@ -509,6 +509,16 @@ public:
 	{
 
 	}
+
+	virtual EColorSpaceAndEOTF RHIGetColorSpace(FRHIViewport* Viewport ) final override
+	{
+		return EColorSpaceAndEOTF::EColorSpace_Rec709;
+	}
+
+	virtual void RHICheckViewportHDRStatus(FRHIViewport* Viewport) final override
+	{
+	}
+
 	virtual void RHITick(float DeltaTime) final override
 	{
 
diff --git a/Engine/Source/Runtime/RHI/Private/DynamicRHI.cpp b/Engine/Source/Runtime/RHI/Private/DynamicRHI.cpp
--- a/Engine/Source/Runtime/RHI/Private/DynamicRHI.cpp
+++ b/Engine/Source/Runtime/RHI/Private/DynamicRHI.cpp
@@ -419,3 +419,13 @@ FRenderQueryPoolRHIRef RHICreateRenderQueryPool(ERenderQueryType QueryType, uint
 {
 	return GDynamicRHI->RHICreateRenderQueryPool(QueryType, NumQueries);
 }
+
+EColorSpaceAndEOTF FDynamicRHI::RHIGetColorSpace(FRHIViewport* Viewport)
+{
+	return EColorSpaceAndEOTF::ERec709_sRGB;
+}
+
+void FDynamicRHI::RHICheckViewportHDRStatus(FRHIViewport* Viewport)
+{
+}
+
diff --git a/Engine/Source/Runtime/RHI/Public/DynamicRHI.h b/Engine/Source/Runtime/RHI/Public/DynamicRHI.h
--- a/Engine/Source/Runtime/RHI/Public/DynamicRHI.h
+++ b/Engine/Source/Runtime/RHI/Public/DynamicRHI.h
@@ -991,6 +991,12 @@ public:
 		RHIResizeViewport(Viewport, SizeX, SizeY, bIsFullscreen);
 	}
 
+	// Return what colour space the viewport is in. Used for HDR displays
+	virtual EColorSpaceAndEOTF RHIGetColorSpace(FRHIViewport* Viewport);
+
+	// Tests the viewport to see if its HDR status has changed. This is usually tested after a window has been moved
+	virtual void RHICheckViewportHDRStatus(FRHIViewport* Viewport);
+
 	//  must be called from the main thread.
 	// FlushType: Thread safe
 	virtual void RHITick(float DeltaTime) = 0;
@@ -1389,6 +1395,16 @@ FORCEINLINE void RHIResizeViewport(FRHIViewport* Viewport, uint32 SizeX, uint32
 	GDynamicRHI->RHIResizeViewport(Viewport, SizeX, SizeY, bIsFullscreen, PreferredPixelFormat);
 }
 
+FORCEINLINE EColorSpaceAndEOTF RHIGetColorSpace(FRHIViewport* Viewport)
+{
+	return GDynamicRHI->RHIGetColorSpace(Viewport);
+}
+
+FORCEINLINE void RHICheckViewportHDRStatus(FRHIViewport* Viewport)
+{
+	GDynamicRHI->RHICheckViewportHDRStatus(Viewport);
+}
+
 FORCEINLINE void RHITick(float DeltaTime)
 {
 	GDynamicRHI->RHITick(DeltaTime);
diff --git a/Engine/Source/Runtime/RHI/Public/RHI.h b/Engine/Source/Runtime/RHI/Public/RHI.h
--- a/Engine/Source/Runtime/RHI/Public/RHI.h
+++ b/Engine/Source/Runtime/RHI/Public/RHI.h
@@ -1072,6 +1072,33 @@ enum class EClearBinding
 	EDepthStencilBound, //target has a depthstencil value bound.  Clears will use the bound values and do hardware clears.
 };
 
+
+enum class EColorSpaceAndEOTF
+{
+	EUnknown = 0,
+
+	EColorSpace_Rec709  = 1,		// Color Space Uses Rec 709  Primaries
+	EColorSpace_Rec2020 = 2,		// Color Space Uses Rec 2020 Primaries
+	EColorSpace_DCIP3   = 3,		// Color Space Uses DCI-P3   Primaries
+	EEColorSpace_MASK   = 0xf,
+
+	EEOTF_Linear		= 1 << 4,   // Transfer Function Uses Linear Encoding
+	EEOTF_sRGB			= 2 << 4,	// Transfer Function Uses sRGB Encoding
+	EEOTF_PQ			= 3 << 4,	// Transfer Function Uses PQ Encoding
+	EEOTF_MASK			= 0xf << 4,
+
+	ERec709_sRGB		= EColorSpace_Rec709  | EEOTF_sRGB,
+	ERec709_Linear		= EColorSpace_Rec709  | EEOTF_Linear,
+	
+	ERec2020_PQ			= EColorSpace_Rec2020 | EEOTF_PQ,
+	ERec2020_Linear		= EColorSpace_Rec2020 | EEOTF_Linear,
+	
+	EDCIP3_PQ			= EColorSpace_DCIP3 | EEOTF_PQ,
+	EDCIP3_Linear		= EColorSpace_DCIP3 | EEOTF_Linear,
+	
+};
+
+
 struct FClearValueBinding
 {
 	struct DSVAlue
diff --git a/Engine/Source/Runtime/Renderer/Private/PostProcess/PostProcessCombineLUTs.cpp b/Engine/Source/Runtime/Renderer/Private/PostProcess/PostProcessCombineLUTs.cpp
--- a/Engine/Source/Runtime/Renderer/Private/PostProcess/PostProcessCombineLUTs.cpp
+++ b/Engine/Source/Runtime/Renderer/Private/PostProcess/PostProcessCombineLUTs.cpp
@@ -261,6 +261,12 @@ public:
 
 			float Gamma = CVarOutputGamma->GetValueOnRenderThread();
 
+			// If the Viewfamily is requesting HDR for editor usage then override the CVAR setting
+			if (ViewFamily.bIsHDR)
+			{
+				OutputDeviceValue = 3;
+			}
+
 			if (PLATFORM_APPLE && Gamma == 0.0f)
 			{
 				Gamma = 2.2f;
diff --git a/Engine/Source/Runtime/Renderer/Private/PostProcess/PostProcessTonemap.cpp b/Engine/Source/Runtime/Renderer/Private/PostProcess/PostProcessTonemap.cpp
--- a/Engine/Source/Runtime/Renderer/Private/PostProcess/PostProcessTonemap.cpp
+++ b/Engine/Source/Runtime/Renderer/Private/PostProcess/PostProcessTonemap.cpp
@@ -514,6 +514,7 @@ public:
 		ScreenPosToScenePixel.Bind(ParameterMap, TEXT("ScreenPosToScenePixel"));
 		SceneUVMinMax.Bind(ParameterMap, TEXT("SceneUVMinMax"));
 		SceneBloomUVMinMax.Bind(ParameterMap, TEXT("SceneBloomUVMinMax"));
+		HDRParams.Bind(ParameterMap, TEXT("HDRParams"));
 	}
 	
 	template <typename TRHICmdList, typename TRHIShader>
@@ -549,6 +550,20 @@ public:
 			SetShaderValue(RHICmdList, ShaderRHI, TonemapperParams, Value);
 		}
 
+		{
+			float EditorNITLevel = 160.0f;
+#if WITH_EDITOR
+			static auto CVarHDRNITLevel = IConsoleManager::Get().FindConsoleVariable(TEXT("Editor.HDRNITLevel"));
+			if (CVarHDRNITLevel)
+			{
+				EditorNITLevel = CVarHDRNITLevel->GetFloat();
+			}
+#endif
+			FVector4 Value(	ViewFamily.bIsHDR ? 1.0f : 0.0f, EditorNITLevel, 0.0f, 0.0f);
+			SetShaderValue(RHICmdList, ShaderRHI, HDRParams, Value);
+		}
+
+
 		FVector GrainValue;
 		GrainPostSettings(&GrainValue, &Settings);
 		SetShaderValue(RHICmdList, ShaderRHI, GrainScaleBiasJitter, GrainValue);
@@ -695,7 +710,7 @@ public:
 		Ar << P.ColorGradingLUT << P.ColorGradingLUTSampler;
 		Ar << P.SceneUVMinMax << P.SceneBloomUVMinMax;
 		Ar << P.ChromaticAberrationParams << P.ScreenPosToScenePixel;
-
+		Ar << P.HDRParams;
 		return Ar;
 	}
 
@@ -713,6 +728,7 @@ public:
 
 	FShaderParameter SceneUVMinMax;
 	FShaderParameter SceneBloomUVMinMax;
+	FShaderParameter HDRParams;
 };
 
 
@@ -1029,13 +1045,21 @@ void FRCPassPostProcessTonemap::Process(FRenderingCompositePassContext& Context)
 			DesktopPermutationVector.Set<TonemapperPermutation::FTonemapperColorFringeDim>(View.FinalPostProcessSettings.SceneFringeIntensity > 0.01f);
 		}
 
+		FTonemapperOutputDevice deviceValue = GetOutputDeviceValue();
+
+		// Override based on the ViewFamily HDR field for editor HDR
+		if (ViewFamily.bIsHDR)
+		{
+			deviceValue = FTonemapperOutputDevice::ACES1000nitST2084;
+		}
+		
 		if (ViewFamily.SceneCaptureSource == SCS_FinalColorHDR)
 		{
 			DesktopPermutationVector.Set<TonemapperPermutation::FTonemapperOutputDeviceDim>(FTonemapperOutputDevice::LinearNoToneCurve);
 		}
 		else
 		{
-			DesktopPermutationVector.Set<TonemapperPermutation::FTonemapperOutputDeviceDim>(GetOutputDeviceValue());
+			DesktopPermutationVector.Set<TonemapperPermutation::FTonemapperOutputDeviceDim>(deviceValue);
 		}
 		DesktopPermutationVector = TonemapperPermutation::RemapPermutation(DesktopPermutationVector, Context.GetFeatureLevel());
 	}
diff --git a/Engine/Source/Runtime/Renderer/Private/PostProcess/PostProcessing.cpp b/Engine/Source/Runtime/Renderer/Private/PostProcess/PostProcessing.cpp
--- a/Engine/Source/Runtime/Renderer/Private/PostProcess/PostProcessing.cpp
+++ b/Engine/Source/Runtime/Renderer/Private/PostProcess/PostProcessing.cpp
@@ -1376,8 +1376,7 @@ void FPostProcessing::Process(FRHICommandListImmediate& RHICmdList, const FViewI
 		const bool bHDROutputEnabled = GRHISupportsHDROutput && IsHDREnabled();
 
 		static const auto CVarDumpFramesAsHDR = IConsoleManager::Get().FindTConsoleVariableDataInt(TEXT("r.BufferVisualizationDumpFramesAsHDR"));
-    	const bool bHDRTonemapperOutput = bAllowTonemapper && (View.Family->SceneCaptureSource == SCS_FinalColorHDR || GetHighResScreenshotConfig().bCaptureHDR || CVarDumpFramesAsHDR->GetValueOnRenderThread() || bHDROutputEnabled);
-
+    	const bool bHDRTonemapperOutput = bAllowTonemapper && (View.Family->SceneCaptureSource == SCS_FinalColorHDR || GetHighResScreenshotConfig().bCaptureHDR || CVarDumpFramesAsHDR->GetValueOnRenderThread() || bHDROutputEnabled || View.Family->bIsHDR);
 
 		FRCPassPostProcessTonemap* Tonemapper = 0;
 
diff --git a/Engine/Source/Runtime/Slate/Private/Framework/Application/SlateApplication.cpp b/Engine/Source/Runtime/Slate/Private/Framework/Application/SlateApplication.cpp
--- a/Engine/Source/Runtime/Slate/Private/Framework/Application/SlateApplication.cpp
+++ b/Engine/Source/Runtime/Slate/Private/Framework/Application/SlateApplication.cpp
@@ -6995,8 +6995,17 @@ bool FSlateApplication::BeginReshapingWindow( const TSharedRef< FGenericWindow >
 	return false;
 }
 
-void FSlateApplication::FinishedReshapingWindow( const TSharedRef< FGenericWindow >& PlatformWindow )
+void FSlateApplication::FinishedReshapingWindow(const TSharedRef< FGenericWindow >& PlatformWindow)
 {
+#if WITH_EDITOR
+	TSharedPtr< SWindow > Window = FSlateWindowHelper::FindWindowByPlatformWindow(SlateWindows, PlatformWindow);
+
+	if (Window.IsValid())
+	{
+		Renderer->OnWindowFinishReshaped(Window);
+	}
+#endif
+
 	if (ThrottleHandle.IsValid())
 	{
 		FSlateThrottleManager::Get().LeaveResponsiveMode(ThrottleHandle);
diff --git a/Engine/Source/Runtime/SlateCore/Public/Rendering/SlateRenderer.h b/Engine/Source/Runtime/SlateCore/Public/Rendering/SlateRenderer.h
--- a/Engine/Source/Runtime/SlateCore/Public/Rendering/SlateRenderer.h
+++ b/Engine/Source/Runtime/SlateCore/Public/Rendering/SlateRenderer.h
@@ -308,6 +308,9 @@ public:
 	/** Called when a window is destroyed to give the renderer a chance to free resources */
 	virtual void OnWindowDestroyed( const TSharedRef<SWindow>& InWindow ) = 0;
 	
+	/** Called when a window is finished being reshaped */
+	virtual void OnWindowFinishReshaped(const TSharedPtr<SWindow>& InWindow) {};
+
 	/**
 	 * Returns the viewport rendering resource (backbuffer) for the provided window
 	 *
diff --git a/Engine/Source/Runtime/SlateCore/Public/Widgets/SWindow.h b/Engine/Source/Runtime/SlateCore/Public/Widgets/SWindow.h
--- a/Engine/Source/Runtime/SlateCore/Public/Widgets/SWindow.h
+++ b/Engine/Source/Runtime/SlateCore/Public/Widgets/SWindow.h
@@ -753,6 +753,13 @@ public:
 		bIsMirrorWindow = bSetMirrorWindow;
 	}
 	
+	void SetIsHDR(bool bHDR)
+	{
+		bIsHDR = bHDR;
+	}
+
+	bool GetIsHDR() const { return bIsHDR; }
+
 	bool IsVirtualWindow() const { return bVirtualWindow; }
 
 	bool IsMirrorWindow()
@@ -1015,6 +1022,9 @@ protected:
 	/** True if the window should preserve its aspect ratio when resized by user */
 	bool bShouldPreserveAspectRatio : 1;
 
+	/** True if the window is being displayed on a HDR capable monitor */
+	bool bIsHDR : 1;
+
 	bool bManualManageDPI : 1;
 
 	/** True if this window allows global invalidation of its contents */
diff --git a/Engine/Source/Runtime/SlateRHIRenderer/Private/SlateRHIRenderer.cpp b/Engine/Source/Runtime/SlateRHIRenderer/Private/SlateRHIRenderer.cpp
--- a/Engine/Source/Runtime/SlateRHIRenderer/Private/SlateRHIRenderer.cpp
+++ b/Engine/Source/Runtime/SlateRHIRenderer/Private/SlateRHIRenderer.cpp
@@ -333,6 +333,10 @@ void FSlateRHIRenderer::CreateViewport(const TSharedRef<SWindow> Window)
 		NewInfo->ViewportRHI = RHICreateViewport( NewInfo->OSWindow, Width, Height, bFullscreen, NewInfo->PixelFormat );
 		NewInfo->bFullscreen = bFullscreen;
 
+		// Was the window created on a HDR compatible display?
+		NewInfo->bHDREnabled = RHIGetColorSpace(NewInfo->ViewportRHI) != EColorSpaceAndEOTF::ERec709_sRGB ;
+		Window->SetIsHDR(NewInfo->bHDREnabled);
+
 		WindowToViewportInfo.Add(&Window.Get(), NewInfo);
 
 		BeginInitResource(NewInfo);
@@ -490,6 +494,13 @@ void FSlateRHIRenderer::OnWindowDestroyed(const TSharedRef<SWindow>& InWindow)
 	WindowToViewportInfo.Remove(&InWindow.Get());
 }
 
+/** Called when a window is Finished being Reshaped - Currently need to check if its HDR status has changed */
+void FSlateRHIRenderer::OnWindowFinishReshaped(const TSharedPtr<SWindow>& InWindow)
+{
+	FViewportInfo* ViewInfo = WindowToViewportInfo.FindRef(InWindow.Get());
+	RHICheckViewportHDRStatus(ViewInfo->ViewportRHI);
+}
+
 // Limited platform support for HDR UI composition
 bool SupportsUICompositionRendering(const EShaderPlatform Platform)
 {
@@ -659,6 +670,15 @@ static FAutoConsoleVariableRef CVarSlateWireframe(TEXT("Slate.ShowWireFrame"), S
 void FSlateRHIRenderer::DrawWindow_RenderThread(FRHICommandListImmediate& RHICmdList, FViewportInfo& ViewportInfo, FSlateWindowElementList& WindowElementList, const struct FSlateDrawWindowCommandParams& DrawCommandParams)
 {
 	LLM_SCOPE(ELLMTag::SceneRender);
+	
+	bool bRenderOffscreen = false;	// Render to an offscreen texture which can then be finally color converted at the end.
+	
+#if WITH_EDITOR
+	if (RHIGetColorSpace(ViewportInfo.ViewportRHI) != EColorSpaceAndEOTF::ERec709_sRGB)
+	{
+		bRenderOffscreen = true;
+	}
+#endif
 
 	FMemMark MemMark(FMemStack::Get());
 
@@ -772,6 +792,29 @@ void FSlateRHIRenderer::DrawWindow_RenderThread(FRHICommandListImmediate& RHICmd
 				BackBuffer = ViewportInfo.UITargetRT->GetRenderTargetItem().TargetableTexture->GetTexture2D();
 			}
 
+#if WITH_EDITOR
+			TRefCountPtr<IPooledRenderTarget> HDRRenderRT;
+
+			if (bRenderOffscreen )
+			{
+				FPooledRenderTargetDesc Desc(FPooledRenderTargetDesc::Create2DDesc(FIntPoint(ViewportWidth, ViewportHeight),
+					PF_FloatRGBA,
+					FClearValueBinding::Transparent,
+					TexCreate_None,
+					TexCreate_ShaderResource | TexCreate_RenderTargetable,
+					false,
+					1,
+					true,
+					true));
+
+				GRenderTargetPool.FindFreeElement(RHICmdList, Desc, HDRRenderRT, TEXT("HDRTargetRT"));
+
+				FResolveParams ResolveParams;
+				RHICmdList.CopyToResolveTarget(FinalBuffer, FinalBuffer, ResolveParams);
+				BackBuffer = HDRRenderRT->GetRenderTargetItem().TargetableTexture->GetTexture2D();
+			}
+#endif
+
 			if (SlateWireFrame)
 			{
 				bClear = true;
@@ -821,6 +864,7 @@ void FSlateRHIRenderer::DrawWindow_RenderThread(FRHICommandListImmediate& RHICmd
 
 						FSlateRenderingParams RenderParams(ViewMatrix * ViewportInfo.ProjectionMatrix, DrawCommandParams.WorldTimeSeconds, DrawCommandParams.DeltaTimeSeconds, DrawCommandParams.RealTimeSeconds);
 						RenderParams.bWireFrame = !!SlateWireFrame;
+						RenderParams.bIsHDR     = ViewportInfo.bHDREnabled;
 
 						FTexture2DRHIRef EmptyTarget;
 
@@ -969,6 +1013,53 @@ void FSlateRHIRenderer::DrawWindow_RenderThread(FRHICommandListImmediate& RHICmd
 				BackBuffer = FinalBuffer;
 			} //bCompositeUI
 
+
+#if WITH_EDITOR
+			if (bRenderOffscreen)
+			{
+				const auto FeatureLevel = GMaxRHIFeatureLevel;
+				auto ShaderMap = GetGlobalShaderMap(FeatureLevel);
+
+				FRHIRenderPassInfo RPInfo(FinalBuffer, ERenderTargetActions::Load_Store);
+				RHICmdList.BeginRenderPass(RPInfo, TEXT("SlateComposite"));
+
+				FGraphicsPipelineStateInitializer GraphicsPSOInit;
+				RHICmdList.ApplyCachedRenderTargets(GraphicsPSOInit);
+				GraphicsPSOInit.BlendState			= TStaticBlendState<>::GetRHI();
+				GraphicsPSOInit.RasterizerState		= TStaticRasterizerState<>::GetRHI();
+				GraphicsPSOInit.DepthStencilState	= TStaticDepthStencilState<false, CF_Always>::GetRHI();
+
+				// ST2084 (PQ) encoding
+				TShaderMapRef<FHDREditorConvertPS> PixelShader(ShaderMap);
+				TShaderMapRef<FScreenVS> VertexShader(ShaderMap);
+
+				GraphicsPSOInit.BoundShaderState.VertexDeclarationRHI = GFilterVertexDeclaration.VertexDeclarationRHI;
+				GraphicsPSOInit.BoundShaderState.VertexShaderRHI = GETSAFERHISHADER_VERTEX(*VertexShader);
+				GraphicsPSOInit.BoundShaderState.PixelShaderRHI = GETSAFERHISHADER_PIXEL(*PixelShader);
+				GraphicsPSOInit.PrimitiveType = PT_TriangleList;
+
+				SetGraphicsPipelineState(RHICmdList, GraphicsPSOInit);
+
+				PixelShader->SetParameters(RHICmdList, HDRRenderRT->GetRenderTargetItem().TargetableTexture );
+
+				static const FName RendererModuleName("Renderer");
+				IRendererModule& RendererModule = FModuleManager::GetModuleChecked<IRendererModule>(RendererModuleName);
+
+				RendererModule.DrawRectangle(
+					RHICmdList,
+					0, 0,
+					ViewportWidth, ViewportHeight,
+					0, 0,
+					ViewportWidth, ViewportHeight,
+					FIntPoint(ViewportWidth, ViewportHeight),
+					FIntPoint(ViewportWidth, ViewportHeight),
+					*VertexShader,
+					EDRF_UseTriangleOptimization);
+				
+				RHICmdList.EndRenderPass();
+				BackBuffer = FinalBuffer;
+			}
+#endif
 			if (!bRenderedStereo && GEngine && IsValidRef(ViewportInfo.GetRenderTargetTexture()) && GEngine->StereoRenderingDevice.IsValid())
 			{
 				const FVector2D WindowSize = WindowElementList.GetWindowSize();
@@ -1124,6 +1215,10 @@ void FSlateRHIRenderer::DrawWindows_Private(FSlateDrawBuffer& WindowDrawBuffer)
 				// The viewport had better exist at this point  
 				FViewportInfo* ViewInfo = WindowToViewportInfo.FindChecked(Window);
 
+				// Cache off the HDR status
+				ViewInfo->bHDREnabled = RHIGetColorSpace(ViewInfo->ViewportRHI) != EColorSpaceAndEOTF::ERec709_sRGB;
+				Window->SetIsHDR(ViewInfo->bHDREnabled);
+
 				if (Window->IsViewportSizeDrivenByWindow())
 				{
 					// Resize the viewport if needed
diff --git a/Engine/Source/Runtime/SlateRHIRenderer/Private/SlateRHIRenderer.h b/Engine/Source/Runtime/SlateRHIRenderer/Private/SlateRHIRenderer.h
--- a/Engine/Source/Runtime/SlateRHIRenderer/Private/SlateRHIRenderer.h
+++ b/Engine/Source/Runtime/SlateRHIRenderer/Private/SlateRHIRenderer.h
@@ -98,6 +98,9 @@ struct FViewportInfo : public FRenderResource
 
 	IViewportRenderTargetProvider* RTProvider;
 	
+	/** Whether is in a HDR Color Space */
+	bool bHDREnabled;
+	
 	/** FRenderResource interface */
 	virtual void InitRHI() override;
 	virtual void ReleaseRHI() override;
@@ -115,7 +118,8 @@ struct FViewportInfo : public FRenderResource
 			bFullscreen(false),
 			PixelFormat(EPixelFormat::PF_Unknown),
 			SDRPixelFormat(EPixelFormat::PF_Unknown),
-			RTProvider(nullptr)
+			RTProvider(nullptr),
+			bHDREnabled(false)
 	{
 	}
 
@@ -173,6 +177,7 @@ public:
 	virtual void Destroy() override;
 	virtual FSlateDrawBuffer& GetDrawBuffer() override;
 	virtual void OnWindowDestroyed( const TSharedRef<SWindow>& InWindow ) override;
+	virtual void OnWindowFinishReshaped(const TSharedPtr<SWindow>& InWindow) override;
 	virtual void RequestResize( const TSharedPtr<SWindow>& Window, uint32 NewWidth, uint32 NewHeight ) override;
 	virtual void CreateViewport( const TSharedRef<SWindow> Window ) override;
 	virtual void UpdateFullscreenState( const TSharedRef<SWindow> Window, uint32 OverrideResX, uint32 OverrideResY ) override;
diff --git a/Engine/Source/Runtime/SlateRHIRenderer/Private/SlateRHIRenderingPolicy.cpp b/Engine/Source/Runtime/SlateRHIRenderer/Private/SlateRHIRenderingPolicy.cpp
--- a/Engine/Source/Runtime/SlateRHIRenderer/Private/SlateRHIRenderingPolicy.cpp
+++ b/Engine/Source/Runtime/SlateRHIRenderer/Private/SlateRHIRenderingPolicy.cpp
@@ -614,7 +614,9 @@ void FSlateRHIRenderingPolicy::DrawElements(
 	// Note that the final editor rendering won't compare 1:1 with 8/10 bit RGBA since blending
 	// of "manually" gammatized values is wrong as there is no de-gammatization of the destination buffer
 	// and re-gammatization of the resulting blending operation in the 8/10 bit RGBA path.
-	const float EngineGamma = (BackBuffer.GetRenderTargetTexture()->GetFormat() == PF_FloatRGBA) ? 1.0 : GEngine ? GEngine->GetDisplayGamma() : 2.2f;
+	// For Editor running in HDR then the gamma needs to be 2.2 and have a float back buffer format.
+
+	const float EngineGamma = ((BackBuffer.GetRenderTargetTexture()->GetFormat() == PF_FloatRGBA)&&(Params.bIsHDR==false)) ? 1.0 : GEngine ? GEngine->GetDisplayGamma() : 2.2f;
 	const float DisplayGamma = bGammaCorrect ? EngineGamma : 1.0f;
 	const float DisplayContrast = GSlateContrast;
 
diff --git a/Engine/Source/Runtime/SlateRHIRenderer/Private/SlateRHIRenderingPolicy.h b/Engine/Source/Runtime/SlateRHIRenderer/Private/SlateRHIRenderingPolicy.h
--- a/Engine/Source/Runtime/SlateRHIRenderer/Private/SlateRHIRenderingPolicy.h
+++ b/Engine/Source/Runtime/SlateRHIRenderer/Private/SlateRHIRenderingPolicy.h
@@ -29,6 +29,7 @@ struct FSlateRenderingParams
 	float CurrentRealTime;
 	bool bAllowSwitchVerticalAxis;
 	bool bWireFrame;
+	bool bIsHDR;
 
 	FSlateRenderingParams(const FMatrix& InViewProjectionMatrix, float InCurrentWorldTime, float InDeltaTimeSeconds, float InCurrentRealTime)
 		: ViewProjectionMatrix(InViewProjectionMatrix)
@@ -38,6 +39,7 @@ struct FSlateRenderingParams
 		, CurrentRealTime(InCurrentRealTime)
 		, bAllowSwitchVerticalAxis(true)
 		, bWireFrame(false)
+		, bIsHDR(false)
 	{
 	}
 };
diff --git a/Engine/Source/Runtime/SlateRHIRenderer/Private/SlateShaders.cpp b/Engine/Source/Runtime/SlateRHIRenderer/Private/SlateShaders.cpp
--- a/Engine/Source/Runtime/SlateRHIRenderer/Private/SlateShaders.cpp
+++ b/Engine/Source/Runtime/SlateRHIRenderer/Private/SlateShaders.cpp
@@ -27,6 +27,9 @@ IMPLEMENT_SHADER_TYPE(, FSlateDebugBatchingPS, TEXT("/Engine/Private/SlateElemen
 	typedef TSlateElementPS<ESlateShader::ShaderType,bDrawDisabledEffect,bUseTextureAlpha> TSlateElementPS##ShaderType##bDrawDisabledEffect##bUseTextureAlpha##A; \
 	IMPLEMENT_SHADER_TYPE(template<>,TSlateElementPS##ShaderType##bDrawDisabledEffect##bUseTextureAlpha##A,TEXT("/Engine/Private/SlateElementPixelShader.usf"),TEXT("Main"),SF_Pixel);
 
+#if WITH_EDITOR
+IMPLEMENT_SHADER_TYPE(, FHDREditorConvertPS, TEXT("/Engine/Private/CompositeUIPixelShader.usf"), TEXT("HDREditorConvert"), SF_Pixel);
+#endif
 /**
 * All the different permutations of shaders used by slate. Uses #defines to avoid dynamic branches
 */
diff --git a/Engine/Source/Runtime/SlateRHIRenderer/Private/SlateShaders.h b/Engine/Source/Runtime/SlateRHIRenderer/Private/SlateShaders.h
--- a/Engine/Source/Runtime/SlateRHIRenderer/Private/SlateShaders.h
+++ b/Engine/Source/Runtime/SlateRHIRenderer/Private/SlateShaders.h
@@ -10,6 +10,7 @@
 #include "GlobalShader.h"
 #include "ShaderParameterUtils.h"
 #include "Rendering/RenderingCommon.h"
+#include "RHIStaticStates.h"
 
 extern EColorVisionDeficiency GSlateColorDeficiencyType;
 extern int32 GSlateColorDeficiencySeverity;
@@ -524,6 +525,66 @@ public:
 };
 
 
+#if WITH_EDITOR
+// Pixel shader to convert UI from linear rec709 to PQ 2020 for HDR monitors
+class FHDREditorConvertPS : public FGlobalShader
+{
+	DECLARE_SHADER_TYPE(FHDREditorConvertPS, Global);
+public:
+
+	static bool ShouldCompilePermutation(const FGlobalShaderPermutationParameters& Parameters)
+	{
+		return IsFeatureLevelSupported(Parameters.Platform, ERHIFeatureLevel::SM5);
+	}
+
+	FHDREditorConvertPS(const ShaderMetaType::CompiledShaderInitializerType& Initializer) :
+		FGlobalShader(Initializer)
+	{
+		SceneTexture.Bind(Initializer.ParameterMap, TEXT("SceneTexture"));
+		SceneSampler.Bind(Initializer.ParameterMap, TEXT("SceneSampler"));
+		UILevel.Bind(Initializer.ParameterMap, TEXT("UILevel"));
+	}
+	FHDREditorConvertPS() {}
+
+	void SetParameters(FRHICommandList& RHICmdList, FRHITexture* SceneTextureRHI)
+	{
+		static const auto CVarOutputDevice = IConsoleManager::Get().FindTConsoleVariableDataInt(TEXT("r.HDR.Display.OutputDevice"));
+
+		SetTextureParameter(RHICmdList, GetPixelShader(), SceneTexture, SceneSampler, TStaticSamplerState<SF_Point>::GetRHI(), SceneTextureRHI);
+		
+		static auto CVarHDRNITLevel = IConsoleManager::Get().FindConsoleVariable(TEXT("Editor.HDRNITLevel"));
+		SetShaderValue(RHICmdList, GetPixelShader(), UILevel, CVarHDRNITLevel->GetFloat());
+	}
+
+	static void ModifyCompilationEnvironment(const FGlobalShaderPermutationParameters& Parameters, FShaderCompilerEnvironment& OutEnvironment)
+	{
+		FGlobalShader::ModifyCompilationEnvironment(Parameters, OutEnvironment);
+	}
+
+	virtual bool Serialize(FArchive& Ar) override
+	{
+		bool bShaderHasOutdatedParameters = FGlobalShader::Serialize(Ar);
+		Ar << SceneTexture << SceneSampler << UILevel;
+		return bShaderHasOutdatedParameters;
+	}
+
+	static const TCHAR* GetSourceFilename()
+	{
+		return TEXT("/Engine/Private/CompositeUIPixelShader.usf");
+	}
+
+	static const TCHAR* GetFunctionName()
+	{
+		return TEXT("HDREditorConvert");
+	}
+
+private:
+	FShaderResourceParameter SceneTexture;
+	FShaderResourceParameter SceneSampler;
+	FShaderParameter UILevel;
+};
+#endif
+
 /** The simple element vertex declaration. */
 extern TGlobalResource<FSlateVertexDeclaration> GSlateVertexDeclaration;
 
diff --git a/Engine/Source/Runtime/Windows/D3D11RHI/Private/D3D11Device.cpp b/Engine/Source/Runtime/Windows/D3D11RHI/Private/D3D11Device.cpp
--- a/Engine/Source/Runtime/Windows/D3D11RHI/Private/D3D11Device.cpp
+++ b/Engine/Source/Runtime/Windows/D3D11RHI/Private/D3D11Device.cpp
@@ -15,6 +15,8 @@
 	#endif
 #include "Windows/HideWindowsPlatformTypes.h"
 
+#include "dxgi1_5.h"
+
 
 bool D3D11RHI_ShouldCreateWithD3DDebug()
 {
@@ -96,6 +98,18 @@ FD3D11DynamicRHI::FD3D11DynamicRHI(IDXGIFactory1* InDXGIFactory1,D3D_FEATURE_LEV
 		GSupportsDepthFetchDuringDepthTest = false;
 	}
 
+   	
+	TRefCountPtr<IDXGIFactory5> Factory5;
+	HRESULT HResult = DXGIFactory1->QueryInterface(IID_PPV_ARGS(Factory5.GetInitReference()));
+	if (SUCCEEDED(HResult))
+	{
+		bDXGISupportsHDR = true;
+	}
+	else
+	{
+		bDXGISupportsHDR = false;
+	}
+
 	ERHIFeatureLevel::Type PreviewFeatureLevel;
 	if (RHIGetPreviewFeatureLevel(PreviewFeatureLevel))
 	{
diff --git a/Engine/Source/Runtime/Windows/D3D11RHI/Private/D3D11RHIPrivate.h b/Engine/Source/Runtime/Windows/D3D11RHI/Private/D3D11RHIPrivate.h
--- a/Engine/Source/Runtime/Windows/D3D11RHI/Private/D3D11RHIPrivate.h
+++ b/Engine/Source/Runtime/Windows/D3D11RHI/Private/D3D11RHIPrivate.h
@@ -486,6 +486,8 @@ public:
 	virtual FViewportRHIRef RHICreateViewport(void* WindowHandle, uint32 SizeX, uint32 SizeY, bool bIsFullscreen, EPixelFormat PreferredPixelFormat) override;
 	virtual void RHIResizeViewport(FRHIViewport* Viewport, uint32 SizeX, uint32 SizeY, bool bIsFullscreen) final override;
 	virtual void RHIResizeViewport(FRHIViewport* Viewport, uint32 SizeX, uint32 SizeY, bool bIsFullscreen, EPixelFormat PreferredPixelFormat) final override;
+	virtual EColorSpaceAndEOTF RHIGetColorSpace(FRHIViewport* Viewport) final override;
+	virtual void RHICheckViewportHDRStatus(FRHIViewport* ViewportRHI) final override;
 	virtual void RHITick(float DeltaTime) final override;
 	virtual void RHISetStreamOutTargets(uint32 NumTargets, FRHIVertexBuffer* const* VertexBuffers,const uint32* Offsets) final override;
 	virtual void RHIBlockUntilGPUIdle() final override;
@@ -781,6 +783,9 @@ protected:
 	/** The global D3D interface. */
 	TRefCountPtr<IDXGIFactory1> DXGIFactory1;
 
+	// Whether HDR is available from the particular DXGI factories available
+	bool bDXGISupportsHDR;
+
 	/** The global D3D device's immediate context */
 	TRefCountPtr<FD3D11DeviceContext> Direct3DDeviceIMContext;
 
diff --git a/Engine/Source/Runtime/Windows/D3D11RHI/Private/D3D11Viewport.cpp b/Engine/Source/Runtime/Windows/D3D11RHI/Private/D3D11Viewport.cpp
--- a/Engine/Source/Runtime/Windows/D3D11RHI/Private/D3D11Viewport.cpp
+++ b/Engine/Source/Runtime/Windows/D3D11RHI/Private/D3D11Viewport.cpp
@@ -263,7 +263,7 @@ void FD3D11Viewport::Resize(uint32 InSizeX, uint32 InSizeY, bool bInIsFullscreen
 			// Resize all existing buffers, don't change count
 			VERIFYD3D11RESIZEVIEWPORTRESULT(SwapChain->ResizeBuffers(0, SizeX, SizeY, RenderTargetFormat, DXGI_SWAP_CHAIN_FLAG_ALLOW_MODE_SWITCH), SizeX, SizeY, RenderTargetFormat, D3DRHI->GetDevice());
 #else
-			VERIFYD3D11RESIZEVIEWPORTRESULT(SwapChain->ResizeBuffers(1, SizeX, SizeY, RenderTargetFormat, DXGI_SWAP_CHAIN_FLAG_ALLOW_MODE_SWITCH), SizeX, SizeY, RenderTargetFormat, D3DRHI->GetDevice());
+			VERIFYD3D11RESIZEVIEWPORTRESULT(SwapChain->ResizeBuffers(BackBufferCount, SizeX, SizeY, RenderTargetFormat, DXGI_SWAP_CHAIN_FLAG_ALLOW_MODE_SWITCH), SizeX, SizeY, RenderTargetFormat, D3DRHI->GetDevice());
 #endif
 
 			if (bInIsFullscreen)
@@ -301,6 +301,9 @@ void FD3D11Viewport::Resize(uint32 InSizeX, uint32 InSizeY, bool bInIsFullscreen
 		D3DRHI->ShutdownHDR();
 	}
 
+	// If the window has been moved or resized it may have moved focus to and from a HDR monitor
+	CheckHDRMonitorStatus();
+	
 	// Create a RHI surface to represent the viewport's back buffer.
 	BackBuffer = GetSwapChainSurface(D3DRHI, PixelFormat, SizeX, SizeY, SwapChain);
 }
@@ -549,6 +552,18 @@ bool FD3D11Viewport::Present(bool bLockToVsync)
 	return bNativelyPresented;
 }
 
+EColorSpaceAndEOTF FD3D11DynamicRHI::RHIGetColorSpace(FRHIViewport* ViewportRHI)
+{
+	FD3D11Viewport* Viewport = ResourceCast(ViewportRHI);
+	return Viewport->GetPixelColorSpace();
+}
+
+void  FD3D11DynamicRHI::RHICheckViewportHDRStatus(FRHIViewport* ViewportRHI)
+{
+	FD3D11Viewport* Viewport = ResourceCast(ViewportRHI);
+	return Viewport->CheckHDRMonitorStatus();
+}
+
 /*=============================================================================
  *	The following RHI functions must be called from the main thread.
  *=============================================================================*/
diff --git a/Engine/Source/Runtime/Windows/D3D11RHI/Private/Windows/WindowsD3D11Viewport.cpp b/Engine/Source/Runtime/Windows/D3D11RHI/Private/Windows/WindowsD3D11Viewport.cpp
--- a/Engine/Source/Runtime/Windows/D3D11RHI/Private/Windows/WindowsD3D11Viewport.cpp
+++ b/Engine/Source/Runtime/Windows/D3D11RHI/Private/Windows/WindowsD3D11Viewport.cpp
@@ -11,7 +11,8 @@
 #include "Windows/AllowWindowsPlatformTypes.h"
 #include <dwmapi.h>
 
-#include "dxgi1_2.h"
+#include "dxgi1_6.h"
+
 
 
 FD3D11Viewport::FD3D11Viewport(FD3D11DynamicRHI* InD3DRHI,HWND InWindowHandle,uint32 InSizeX,uint32 InSizeY,bool bInIsFullscreen, EPixelFormat InPreferredPixelFormat):
@@ -27,6 +28,7 @@ FD3D11Viewport::FD3D11Viewport(FD3D11DynamicRHI* InD3DRHI,HWND InWindowHandle,ui
 	SizeY(InSizeY),
 	bIsFullscreen(bInIsFullscreen),
 	PixelFormat(InPreferredPixelFormat),
+	PixelColorSpace(EColorSpaceAndEOTF::ERec709_sRGB),
 	bIsValid(true),
 	FrameSyncEvent(InD3DRHI)
 {
@@ -107,9 +109,43 @@ FD3D11Viewport::FD3D11Viewport(FD3D11DynamicRHI* InD3DRHI,HWND InWindowHandle,ui
 			}
 		}
 
-		// if stereo was not activated or not enabled in settings
+		// Try and create a swapchain capable of being used on HDR monitors
+		if ((SwapChain == nullptr) && (InD3DRHI->bDXGISupportsHDR))
+		{
+			// Create the swapchain.
+			DXGI_SWAP_CHAIN_DESC1 SwapChainDesc;
+			FMemory::Memzero(&SwapChainDesc, sizeof(DXGI_SWAP_CHAIN_DESC1));
+			SwapChainDesc.Width = SizeX;
+			SwapChainDesc.Height = SizeY;
+			SwapChainDesc.SampleDesc.Count = 1;
+			SwapChainDesc.SampleDesc.Quality = 0;
+			SwapChainDesc.Format = DXGI_FORMAT_R10G10B10A2_UNORM;
+			SwapChainDesc.BufferUsage = DXGI_USAGE_RENDER_TARGET_OUTPUT | DXGI_USAGE_SHADER_INPUT;
+
+			DXGI_SWAP_CHAIN_FULLSCREEN_DESC FSSwapChainDesc = {};
+			FSSwapChainDesc.Windowed = !bIsFullscreen;
+
+			// Needed for HDR
+			BackBufferCount = 2;
+			SwapChainDesc.SwapEffect = DXGI_SWAP_EFFECT_FLIP_DISCARD;
+			SwapChainDesc.BufferCount = BackBufferCount;
+			SwapChainDesc.Flags = DXGI_SWAP_CHAIN_FLAG_ALLOW_MODE_SWITCH;
+
+			IDXGISwapChain1* SwapChain1 = nullptr;
+			IDXGIFactory2* Factory2 = (IDXGIFactory2*)D3DRHI->GetFactory();
+
+			VERIFYD3D11RESULT_EX((Factory2->CreateSwapChainForHwnd(D3DRHI->GetDevice(), WindowHandle, &SwapChainDesc, &FSSwapChainDesc, nullptr, &SwapChain1)), D3DRHI->GetDevice());
+			SwapChain1->QueryInterface(__uuidof(IDXGISwapChain1), (void**)SwapChain.GetInitReference());
+
+			// See if we are running on a HDR monitor 
+			CheckHDRMonitorStatus();
+		}
+
+
 		if (SwapChain == nullptr)
 		{
+			BackBufferCount = 1;
+
 			// Create the swapchain.
 			DXGI_SWAP_CHAIN_DESC SwapChainDesc;
 			FMemory::Memzero(&SwapChainDesc, sizeof(DXGI_SWAP_CHAIN_DESC));
@@ -120,7 +156,7 @@ FD3D11Viewport::FD3D11Viewport(FD3D11DynamicRHI* InD3DRHI,HWND InWindowHandle,ui
 			SwapChainDesc.SampleDesc.Quality = 0;
 			SwapChainDesc.BufferUsage = DXGI_USAGE_RENDER_TARGET_OUTPUT | DXGI_USAGE_SHADER_INPUT;
 			// 1:single buffering, 2:double buffering, 3:triple buffering
-			SwapChainDesc.BufferCount = 1;
+			SwapChainDesc.BufferCount = BackBufferCount;
 			SwapChainDesc.OutputWindow = WindowHandle;
 			SwapChainDesc.Windowed = !bIsFullscreen;
 			// DXGI_SWAP_EFFECT_DISCARD / DXGI_SWAP_EFFECT_SEQUENTIAL
@@ -142,6 +178,65 @@ FD3D11Viewport::FD3D11Viewport(FD3D11DynamicRHI* InD3DRHI,HWND InWindowHandle,ui
 	BeginInitResource(&FrameSyncEvent);
 }
 
+// When a window has moved or resized we need to check whether it is on a HDR monitor or not. Set the correct color space of the monitor
+void FD3D11Viewport::CheckHDRMonitorStatus()
+{
+#if WITH_EDITOR
+
+	DXGI_COLOR_SPACE_TYPE colorSpace = DXGI_COLOR_SPACE_RGB_FULL_G22_NONE_P709;
+
+	static auto CVarHDREnable = IConsoleManager::Get().FindConsoleVariable(TEXT("Editor.HDRSupport"));
+	if (CVarHDREnable->GetInt() != 0)
+	{
+		FlushRenderingCommands();
+
+		if (SwapChain)
+		{
+			TRefCountPtr<IDXGIOutput> Output;
+			if (SUCCEEDED(SwapChain->GetContainingOutput(Output.GetInitReference())))
+			{
+				TRefCountPtr<IDXGIOutput6> Output6;
+				if (SUCCEEDED(Output->QueryInterface(IID_PPV_ARGS(Output6.GetInitReference()))))
+				{
+					DXGI_OUTPUT_DESC1 desc;
+					Output6->GetDesc1(&desc);
+
+					if (desc.ColorSpace == DXGI_COLOR_SPACE_RGB_FULL_G2084_NONE_P2020)
+					{
+						// Display output is HDR10.
+						colorSpace = DXGI_COLOR_SPACE_RGB_FULL_G2084_NONE_P2020;
+					}
+				}
+			}
+		}
+
+		TRefCountPtr<IDXGISwapChain3> swapChain3;
+
+		if (SUCCEEDED(SwapChain->QueryInterface(IID_PPV_ARGS(swapChain3.GetInitReference()))))
+		{
+			UINT colorSpaceSupport = 0;
+		
+			if (SUCCEEDED(swapChain3->CheckColorSpaceSupport(colorSpace, &colorSpaceSupport)) && 
+				(colorSpaceSupport & DXGI_SWAP_CHAIN_COLOR_SPACE_SUPPORT_FLAG_PRESENT))
+			{
+				swapChain3->SetColorSpace1(colorSpace);
+			}
+		}
+	}
+	
+	if (colorSpace == DXGI_COLOR_SPACE_RGB_FULL_G2084_NONE_P2020)
+	{
+		PixelColorSpace = EColorSpaceAndEOTF::ERec2020_PQ;
+	}
+	else
+	{
+		PixelColorSpace =  EColorSpaceAndEOTF::ERec709_sRGB;
+	}
+#else
+	PixelColorSpace =  EColorSpaceAndEOTF::ERec709_sRGB;
+#endif
+}
+
 void FD3D11Viewport::ConditionalResetSwapChain(bool bIgnoreFocus)
 {
 	if (!bIsValid)
diff --git a/Engine/Source/Runtime/Windows/D3D11RHI/Public/D3D11Viewport.h b/Engine/Source/Runtime/Windows/D3D11RHI/Public/D3D11Viewport.h
--- a/Engine/Source/Runtime/Windows/D3D11RHI/Public/D3D11Viewport.h
+++ b/Engine/Source/Runtime/Windows/D3D11RHI/Public/D3D11Viewport.h
@@ -74,6 +74,12 @@ public:
      */
 	void ConditionalResetSwapChain(bool bIgnoreFocus);
 
+	/**
+	 * Called whenever the Viewport is moved to see if it has moved between HDR or LDR monitors
+	 */
+	void CheckHDRMonitorStatus();
+
+
 	/** Presents the swap chain. 
 	 * Returns true if Present was done by Engine.
  	 */
@@ -82,6 +88,7 @@ public:
 	// Accessors.
 	FIntPoint GetSizeXY() const { return FIntPoint(SizeX, SizeY); }
 	FD3D11Texture2D* GetBackBuffer() const { return BackBuffer; }
+	EColorSpaceAndEOTF GetPixelColorSpace() const { return PixelColorSpace; }
 
 	void WaitForFrameEventCompletion()
 	{
@@ -133,8 +140,10 @@ protected:
 	uint32 MaximumFrameLatency;
 	uint32 SizeX;
 	uint32 SizeY;
+	uint32 BackBufferCount;
 	bool bIsFullscreen;
 	EPixelFormat PixelFormat;
+	EColorSpaceAndEOTF PixelColorSpace;
 	bool bIsValid;
 #if PLATFORM_HOLOLENS
 	TRefCountPtr<IDXGISwapChain1> SwapChain;
