Slate - For RHI based material shaders, this CL eliminates several permutations that we didn't need, but were using in some cases (removal of those in the editor was a previous CL).  This eliminates 3 "Disabled Effect" permutations per shaders, this is now a dyanmic branch in the shader - It's still a permutation in slate global shaders.  This also removes all the FLocalVertexFactory permutations, as well as the ConvertToUniformMesh permutation, both of which shouldn't have ever been on for UII Materials - but the editor was using the FLocalVertexFactory one (but those cases should all be fixed now).

Matt.Kuhlenschmidt

diff --git a/Engine/Source/Runtime/Engine/Private/LocalVertexFactory.cpp b/Engine/Source/Runtime/Engine/Private/LocalVertexFactory.cpp
--- a/Engine/Source/Runtime/Engine/Private/LocalVertexFactory.cpp
+++ b/Engine/Source/Runtime/Engine/Private/LocalVertexFactory.cpp
@@ -181,6 +181,11 @@ void FLocalVertexFactoryShaderParameters::GetElementShaderBindings(
  */
 bool FLocalVertexFactory::ShouldCompilePermutation(EShaderPlatform Platform, const class FMaterial* Material, const class FShaderType* ShaderType)
 {
+	if (Material->GetMaterialDomain() == MD_UI)
+	{
+		return false;
+	}
+
 	return true; 
 }
 
diff --git a/Engine/Source/Runtime/Renderer/Private/ConvertToUniformMesh.cpp b/Engine/Source/Runtime/Renderer/Private/ConvertToUniformMesh.cpp
--- a/Engine/Source/Runtime/Renderer/Private/ConvertToUniformMesh.cpp
+++ b/Engine/Source/Runtime/Renderer/Private/ConvertToUniformMesh.cpp
@@ -306,7 +306,16 @@ public:
 
 	static bool ShouldCompilePermutation(const FMaterialShaderPermutationParameters& Parameters)
 	{
-		//@todo - lit materials only 
+		if (Material->GetMaterialDomain() == MD_UI)
+		{
+			return false;
+		}
+
+		if (Material->GetShadingModels().IsUnlit())
+		{
+			return false;
+		}
+
 		return IsFeatureLevelSupported(Parameters.Platform, ERHIFeatureLevel::SM5) && DoesPlatformSupportDistanceFieldGI(Parameters.Platform);
 	}
 
diff --git a/Engine/Source/Runtime/SlateRHIRenderer/Private/SlateMaterialShader.cpp b/Engine/Source/Runtime/SlateRHIRenderer/Private/SlateMaterialShader.cpp
--- a/Engine/Source/Runtime/SlateRHIRenderer/Private/SlateMaterialShader.cpp
+++ b/Engine/Source/Runtime/SlateRHIRenderer/Private/SlateMaterialShader.cpp
@@ -77,6 +77,7 @@ FSlateMaterialShaderPS::FSlateMaterialShaderPS(const FMaterialShaderType::Compil
 {
 	ShaderParams.Bind(Initializer.ParameterMap, TEXT("ShaderParams"));
 	GammaAndAlphaValues.Bind(Initializer.ParameterMap, TEXT("GammaAndAlphaValues"));
+	DrawFlags.Bind(Initializer.ParameterMap, TEXT("DrawFlags"));
 	AdditionalTextureParameter.Bind(Initializer.ParameterMap, TEXT("ElementTexture"));
 	TextureParameterSampler.Bind(Initializer.ParameterMap, TEXT("ElementTextureSampler"));
 }
@@ -139,10 +140,18 @@ void FSlateMaterialShaderPS::SetDisplayGammaAndContrast(FRHICommandList& RHICmdL
 	SetShaderValue(RHICmdList, GetPixelShader(), GammaAndAlphaValues, InGammaValues);
 }
 
+void FSlateMaterialShaderPS::SetDrawFlags(FRHICommandList& RHICmdList, bool bDrawDisabledEffect)
+{
+	FVector4 InDrawFlags(bDrawDisabledEffect ? 1 : 0, 0, 0, 0);
+
+	SetShaderValue(RHICmdList, GetPixelShader(), DrawFlags, InDrawFlags);
+}
+
 bool FSlateMaterialShaderPS::Serialize(FArchive& Ar)
 {
 	bool bShaderHasOutdatedParameters = FMaterialShader::Serialize(Ar);
 	Ar << GammaAndAlphaValues;
+	Ar << DrawFlags;
 	Ar << ShaderParams;
 	Ar << TextureParameterSampler;
 	Ar << AdditionalTextureParameter;
@@ -159,17 +168,12 @@ IMPLEMENT_SLATE_VERTEXMATERIALSHADER_TYPE(true);
 /** Non instancing vertex shader */
 IMPLEMENT_SLATE_VERTEXMATERIALSHADER_TYPE(false);
 
-#define IMPLEMENT_SLATE_MATERIALSHADER_TYPE(ShaderType, bDrawDisabledEffect) \
-	typedef TSlateMaterialShaderPS<ESlateShader::ShaderType, bDrawDisabledEffect> TSlateMaterialShaderPS##ShaderType##bDrawDisabledEffect; \
-	IMPLEMENT_MATERIAL_SHADER_TYPE(template<>, TSlateMaterialShaderPS##ShaderType##bDrawDisabledEffect, TEXT("/Engine/Private/SlateElementPixelShader.usf"), TEXT("Main"), SF_Pixel);
-
-IMPLEMENT_SLATE_MATERIALSHADER_TYPE(Custom, false)
+#define IMPLEMENT_SLATE_MATERIALSHADER_TYPE(ShaderType) \
+	typedef TSlateMaterialShaderPS<ESlateShader::ShaderType> TSlateMaterialShaderPS##ShaderType; \
+	IMPLEMENT_MATERIAL_SHADER_TYPE(template<>, TSlateMaterialShaderPS##ShaderType, TEXT("/Engine/Private/SlateElementPixelShader.usf"), TEXT("Main"), SF_Pixel);
 
-IMPLEMENT_SLATE_MATERIALSHADER_TYPE(Default, true);
-IMPLEMENT_SLATE_MATERIALSHADER_TYPE(Default, false);
-IMPLEMENT_SLATE_MATERIALSHADER_TYPE(Border, true);
-IMPLEMENT_SLATE_MATERIALSHADER_TYPE(Border, false);
-IMPLEMENT_SLATE_MATERIALSHADER_TYPE(GrayscaleFont, true);
-IMPLEMENT_SLATE_MATERIALSHADER_TYPE(GrayscaleFont, false);
-IMPLEMENT_SLATE_MATERIALSHADER_TYPE(ColorFont, true);
-IMPLEMENT_SLATE_MATERIALSHADER_TYPE(ColorFont, false);
+IMPLEMENT_SLATE_MATERIALSHADER_TYPE(Custom)
+IMPLEMENT_SLATE_MATERIALSHADER_TYPE(Default);
+IMPLEMENT_SLATE_MATERIALSHADER_TYPE(Border);
+IMPLEMENT_SLATE_MATERIALSHADER_TYPE(GrayscaleFont);
+IMPLEMENT_SLATE_MATERIALSHADER_TYPE(ColorFont);
diff --git a/Engine/Source/Runtime/SlateRHIRenderer/Private/SlateMaterialShader.h b/Engine/Source/Runtime/SlateRHIRenderer/Private/SlateMaterialShader.h
--- a/Engine/Source/Runtime/SlateRHIRenderer/Private/SlateMaterialShader.h
+++ b/Engine/Source/Runtime/SlateRHIRenderer/Private/SlateMaterialShader.h
@@ -63,10 +63,13 @@ public:
 
 	void SetAdditionalTexture( FRHICommandList& RHICmdList, FRHITexture* InTexture, const FSamplerStateRHIRef SamplerState );
 
+	void SetDrawFlags(FRHICommandList& RHICmdList, bool bDrawDisabledEffect);
+
 	virtual bool Serialize(FArchive& Ar) override;
 
 private:
 	FShaderParameter GammaAndAlphaValues;
+	FShaderParameter DrawFlags;
 	FShaderParameter ShaderParams;
 	/** Extra texture (like a font atlas) to be used in addition to any material textures */
 	FShaderResourceParameter TextureParameterSampler;
@@ -105,7 +108,7 @@ public:
 	}
 };
 
-template<ESlateShader ShaderType,bool bDrawDisabledEffect> 
+template<ESlateShader ShaderType>
 class TSlateMaterialShaderPS : public FSlateMaterialShaderPS
 {
 public:
@@ -129,7 +132,6 @@ public:
 		FSlateMaterialShaderPS::ModifyCompilationEnvironment(Parameters,OutEnvironment);
 
 		OutEnvironment.SetDefine(TEXT("SHADER_TYPE"), (uint32)ShaderType);
-		OutEnvironment.SetDefine(TEXT("DRAW_DISABLED_EFFECT"), (uint32)(bDrawDisabledEffect ? 1 : 0));
 	}
 
 	virtual bool Serialize(FArchive& Ar) override
diff --git a/Engine/Source/Runtime/SlateRHIRenderer/Private/SlateRHIRenderingPolicy.cpp b/Engine/Source/Runtime/SlateRHIRenderer/Private/SlateRHIRenderingPolicy.cpp
--- a/Engine/Source/Runtime/SlateRHIRenderer/Private/SlateRHIRenderingPolicy.cpp
+++ b/Engine/Source/Runtime/SlateRHIRenderer/Private/SlateRHIRenderingPolicy.cpp
@@ -1059,7 +1059,7 @@ void FSlateRHIRenderingPolicy::DrawElements(
 
 					const FMaterial* Material = MaterialRenderProxy->GetMaterial(ActiveSceneView.GetFeatureLevel());
 
-					FSlateMaterialShaderPS* PixelShader = GetMaterialPixelShader(Material, ShaderType, DrawEffects);
+					FSlateMaterialShaderPS* PixelShader = GetMaterialPixelShader(Material, ShaderType);
 
 					const bool bUseInstancing = RenderBatch.InstanceCount > 0 && RenderBatch.InstanceData != nullptr;
 					FSlateMaterialShaderVS* VertexShader = GetMaterialVertexShader(Material, bUseInstancing);
@@ -1117,6 +1117,8 @@ void FSlateRHIRenderingPolicy::DrawElements(
 								const float FinalGamma = EnumHasAnyFlags(DrawFlags, ESlateBatchDrawFlag::ReverseGamma) ? 1.0f / EngineGamma : EnumHasAnyFlags(DrawFlags, ESlateBatchDrawFlag::NoGamma) ? 1.0f : DisplayGamma;
 								const float FinalContrast = EnumHasAnyFlags(DrawFlags, ESlateBatchDrawFlag::NoGamma) ? 1 : DisplayContrast;
 								PixelShader->SetDisplayGammaAndContrast(RHICmdList, FinalGamma, FinalContrast);
+								const bool bDrawDisabled = EnumHasAllFlags(DrawEffects, ESlateDrawEffect::DisabledEffect);
+								PixelShader->SetDrawFlags(RHICmdList, bDrawDisabled);
 
 								if (MaskResource)
 								{
@@ -1420,60 +1422,27 @@ FSlateElementPS* FSlateRHIRenderingPolicy::GetTexturePixelShader( TShaderMap<FGl
 	return PixelShader;
 }
 
-FSlateMaterialShaderPS* FSlateRHIRenderingPolicy::GetMaterialPixelShader( const FMaterial* Material, ESlateShader ShaderType, ESlateDrawEffect DrawEffects )
+FSlateMaterialShaderPS* FSlateRHIRenderingPolicy::GetMaterialPixelShader( const FMaterial* Material, ESlateShader ShaderType)
 {
 	const FMaterialShaderMap* MaterialShaderMap = Material->GetRenderingThreadShaderMap();
 
-	const bool bDrawDisabled = EnumHasAllFlags(DrawEffects, ESlateDrawEffect::DisabledEffect);
-	const bool bUseTextureAlpha = !EnumHasAllFlags(DrawEffects, ESlateDrawEffect::IgnoreTextureAlpha);
-
 	FShader* FoundShader = nullptr;
 	switch (ShaderType)
 	{
 	case ESlateShader::Default:
-		if (bDrawDisabled)
-		{
-			FoundShader = MaterialShaderMap->GetShader(&TSlateMaterialShaderPS<ESlateShader::Default, true>::StaticType);
-		}
-		else
-		{
-			FoundShader = MaterialShaderMap->GetShader(&TSlateMaterialShaderPS<ESlateShader::Default, false>::StaticType);
-		}
+		FoundShader = MaterialShaderMap->GetShader(&TSlateMaterialShaderPS<ESlateShader::Default>::StaticType);
 		break;
 	case ESlateShader::Border:
-		if (bDrawDisabled)
-		{
-			FoundShader = MaterialShaderMap->GetShader(&TSlateMaterialShaderPS<ESlateShader::Border, true>::StaticType);
-		}
-		else
-		{
-			FoundShader = MaterialShaderMap->GetShader(&TSlateMaterialShaderPS<ESlateShader::Border, false>::StaticType);
-		}
+		FoundShader = MaterialShaderMap->GetShader(&TSlateMaterialShaderPS<ESlateShader::Border>::StaticType);
 		break;
 	case ESlateShader::GrayscaleFont:
-		if(bDrawDisabled)
-		{
-			FoundShader = MaterialShaderMap->GetShader(&TSlateMaterialShaderPS<ESlateShader::GrayscaleFont, true>::StaticType);
-		}
-		else
-		{
-			FoundShader = MaterialShaderMap->GetShader(&TSlateMaterialShaderPS<ESlateShader::GrayscaleFont, false>::StaticType);
-		}
+		FoundShader = MaterialShaderMap->GetShader(&TSlateMaterialShaderPS<ESlateShader::GrayscaleFont>::StaticType);
 		break;
 	case ESlateShader::ColorFont:
-		if (bDrawDisabled)
-		{
-			FoundShader = MaterialShaderMap->GetShader(&TSlateMaterialShaderPS<ESlateShader::ColorFont, true>::StaticType);
-		}
-		else
-		{
-			FoundShader = MaterialShaderMap->GetShader(&TSlateMaterialShaderPS<ESlateShader::ColorFont, false>::StaticType);
-		}
+		FoundShader = MaterialShaderMap->GetShader(&TSlateMaterialShaderPS<ESlateShader::ColorFont>::StaticType);
 		break;
 	case ESlateShader::Custom:
-		{
-			FoundShader = MaterialShaderMap->GetShader(&TSlateMaterialShaderPS<ESlateShader::Custom, false>::StaticType);
-		}
+		FoundShader = MaterialShaderMap->GetShader(&TSlateMaterialShaderPS<ESlateShader::Custom>::StaticType);
 		break;
 	default:
 		checkf(false, TEXT("Unsupported Slate shader type for use with materials"));
diff --git a/Engine/Source/Runtime/SlateRHIRenderer/Private/SlateRHIRenderingPolicy.h b/Engine/Source/Runtime/SlateRHIRenderer/Private/SlateRHIRenderingPolicy.h
--- a/Engine/Source/Runtime/SlateRHIRenderer/Private/SlateRHIRenderingPolicy.h
+++ b/Engine/Source/Runtime/SlateRHIRenderer/Private/SlateRHIRenderingPolicy.h
@@ -81,7 +81,7 @@ private:
 	 * @return The pixel shader for use with the shader type and draw effects
 	 */
 	class FSlateElementPS* GetTexturePixelShader( TShaderMap<FGlobalShaderType>* ShaderMap, ESlateShader ShaderType, ESlateDrawEffect DrawEffects );
-	class FSlateMaterialShaderPS* GetMaterialPixelShader( const class FMaterial* Material, ESlateShader ShaderType, ESlateDrawEffect DrawEffects );
+	class FSlateMaterialShaderPS* GetMaterialPixelShader( const class FMaterial* Material, ESlateShader ShaderType);
 	class FSlateMaterialShaderVS* GetMaterialVertexShader( const class FMaterial* Material, bool bUseInstancing );
 
 	/** @return The RHI primitive type from the Slate primitive type */
