Allow separate size settings for grayscale and color font atlases

This allows them to be configured consistently, rather than the color atlas always being 1/4 the size of the grayscale one. The old FontAtlasSize setting is deprecated in favor of GrayscaleFontAtlasSize and ColorFontAtlasSize.

diff --git a/Engine/Config/BaseEditor.ini b/Engine/Config/BaseEditor.ini
--- a/Engine/Config/BaseEditor.ini
+++ b/Engine/Config/BaseEditor.ini
@@ -15,6 +15,10 @@ CheckSingleInfluenceLOD=False
 [LogWindow]
 MaxNumberOfLogLines=2000
 
+[SlateRenderer]
+GrayscaleFontAtlasSize=2048
+ColorFontAtlasSize=512
+
 [/Script/UnrealEd.ThumbnailManager]
 +RenderableThumbnailTypes=(ClassNeedingThumbnailName="/Script/CoreUObject.Class",RendererClassName="/Script/UnrealEd.ClassThumbnailRenderer")
 +RenderableThumbnailTypes=(ClassNeedingThumbnailName="/Script/Engine.Blueprint",RendererClassName="/Script/UnrealEd.BlueprintThumbnailRenderer")
diff --git a/Engine/Config/BaseEngine.ini b/Engine/Config/BaseEngine.ini
--- a/Engine/Config/BaseEngine.ini
+++ b/Engine/Config/BaseEngine.ini
@@ -1841,7 +1841,8 @@ DeviceProfileSelectionModule="ExampleDeviceProfileSelector"
 
 [SlateRenderer]
 TextureAtlasSize=1024
-FontAtlasSize=1024
+GrayscaleFontAtlasSize=1024
+ColorFontAtlasSize=512
 ; This is the number of preallocated vertices to create for use in Slate/UMG UI. A smaller number would indicate more resizes of the vertex buffer which could cause hitches.  A larger number increases memory
 ; Use stat slatememory to see the cost of these buffers
 NumPreallocatedVertices=50000
diff --git a/Engine/Source/Developer/StandaloneRenderer/Private/OpenGL/SlateOpenGLRenderer.cpp b/Engine/Source/Developer/StandaloneRenderer/Private/OpenGL/SlateOpenGLRenderer.cpp
--- a/Engine/Source/Developer/StandaloneRenderer/Private/OpenGL/SlateOpenGLRenderer.cpp
+++ b/Engine/Source/Developer/StandaloneRenderer/Private/OpenGL/SlateOpenGLRenderer.cpp
@@ -47,7 +47,7 @@ private:
 
 	/** Size of each font texture, width and height */
 	static const uint32 GrayscaleTextureSize = 1024;
-	static const uint32 ColorTextureSize = 256;
+	static const uint32 ColorTextureSize = 512;
 };
 
 TSharedRef<FSlateFontServices> CreateOpenGLFontServices()
diff --git a/Engine/Source/Developer/StandaloneRenderer/Private/Windows/D3D/SlateD3DRenderer.cpp b/Engine/Source/Developer/StandaloneRenderer/Private/Windows/D3D/SlateD3DRenderer.cpp
--- a/Engine/Source/Developer/StandaloneRenderer/Private/Windows/D3D/SlateD3DRenderer.cpp
+++ b/Engine/Source/Developer/StandaloneRenderer/Private/Windows/D3D/SlateD3DRenderer.cpp
@@ -105,7 +105,7 @@ private:
 
 	/** Size of each font texture, width and height */
 	static const uint32 GrayscaleTextureSize = 1024;
-	static const uint32 ColorTextureSize = 256;
+	static const uint32 ColorTextureSize = 512;
 };
 
 TSharedRef<FSlateFontServices> CreateD3DFontServices()
diff --git a/Engine/Source/Runtime/SlateRHIRenderer/Private/SlateRHIRendererModule.cpp b/Engine/Source/Runtime/SlateRHIRenderer/Private/SlateRHIRendererModule.cpp
--- a/Engine/Source/Runtime/SlateRHIRenderer/Private/SlateRHIRendererModule.cpp
+++ b/Engine/Source/Runtime/SlateRHIRenderer/Private/SlateRHIRendererModule.cpp
@@ -21,20 +21,26 @@ class FSlateRHIFontAtlasFactory : public ISlateFontAtlasFactory
 public:
 	FSlateRHIFontAtlasFactory()
 	{
-		if (GIsEditor)
-		{
-			GrayscaleAtlasSize = 2048;
-		}
-		else
+		auto GetAtlasSizeFromConfig = [](const TCHAR* InConfigKey, const FString& InConfigFilename, int32& OutAtlasSize)
 		{
-			GrayscaleAtlasSize = 1024;
 			if (GConfig)
 			{
-				GConfig->GetInt(TEXT("SlateRenderer"), TEXT("FontAtlasSize"), GrayscaleAtlasSize, GEngineIni);
-				GrayscaleAtlasSize = FMath::Clamp(GrayscaleAtlasSize, 0, 2048);
+				GConfig->GetInt(TEXT("SlateRenderer"), InConfigKey, OutAtlasSize, InConfigFilename);
+				if (GConfig->GetInt(TEXT("SlateRenderer"), TEXT("FontAtlasSize"), OutAtlasSize, InConfigFilename))
+				{
+					UE_LOG(LogCore, Warning, TEXT("The 'FontAtlasSize' setting for 'SlateRenderer' is deprecated. Use '%s' instead."), InConfigKey);
+				}
 			}
+			OutAtlasSize = FMath::Clamp(OutAtlasSize, 128, 2048);
+		};
+
+		GrayscaleAtlasSize = GIsEditor ? 2048 : 1024;
+		ColorAtlasSize = 512;
+		{
+			const FString& ConfigFilename = GIsEditor ? GEditorIni : GEngineIni;
+			GetAtlasSizeFromConfig(TEXT("GrayscaleFontAtlasSize"), ConfigFilename, GrayscaleAtlasSize);
+			GetAtlasSizeFromConfig(TEXT("ColorFontAtlasSize"), ConfigFilename, ColorAtlasSize);
 		}
-		ColorAtlasSize = GrayscaleAtlasSize / 4;
 	}
 
 	virtual ~FSlateRHIFontAtlasFactory()
