Fixes merge conflict within SlateRHIRenderingPolicy.cpp

diff --git a/Engine/Source/Runtime/SlateRHIRenderer/Private/SlateRHIRenderingPolicy.cpp b/Engine/Source/Runtime/SlateRHIRenderer/Private/SlateRHIRenderingPolicy.cpp
--- a/Engine/Source/Runtime/SlateRHIRenderer/Private/SlateRHIRenderingPolicy.cpp
+++ b/Engine/Source/Runtime/SlateRHIRenderer/Private/SlateRHIRenderingPolicy.cpp
@@ -131,6 +131,8 @@ void FSlateRHIRenderingPolicy::EndDrawingWindows()
 	check( IsInParallelRenderingThread() );
 }
 
+struct FSlateUpdateVertexAndIndexBuffers final : public FRHICommand<FSlateUpdateVertexAndIndexBuffers>
+{
 	TSlateElementVertexBuffer<FSlateVertex>& VertexBuffer;
 	FSlateElementIndexBuffer& IndexBuffer;
 	FSlateBatchData& BatchData;
