Fixed post merge compilation issue.

diff --git a/Engine/Source/Runtime/SlateRHIRenderer/Private/SlateUpdatableBuffer.cpp b/Engine/Source/Runtime/SlateRHIRenderer/Private/SlateUpdatableBuffer.cpp
--- a/Engine/Source/Runtime/SlateRHIRenderer/Private/SlateUpdatableBuffer.cpp
+++ b/Engine/Source/Runtime/SlateRHIRenderer/Private/SlateUpdatableBuffer.cpp
@@ -98,21 +98,7 @@ void FSlateUpdatableInstanceBuffer::UpdateRenderingData_RenderThread(FRHICommand
 	}
 	else
 	{
-		FRHIVertexBuffer* VertexBufferRHI = InstanceBufferResource.VertexBufferRHI;
-		RHICmdList.EnqueueLambda([VertexBufferRHI, &InstanceData = RenderThreadBufferData](FRHICommandListImmediate& InRHICmdList)
-		{
-			// Use different cycle counter to avoid -Wshadow warning on older clang versions for Android
-			SCOPE_CYCLE_COUNTER(STAT_SlateUpdateInstanceBufferLambda);
-
-			int32 RequiredVertexBufferSize = InstanceData.Num() * sizeof(FVector4);
-			uint8* InstanceBufferData = (uint8*)InRHICmdList.LockVertexBuffer(VertexBufferRHI, 0, RequiredVertexBufferSize, RLM_WriteOnly);
-
-			FMemory::Memcpy(InstanceBufferData, InstanceData.GetData(), InstanceData.Num() * sizeof(FVector4));
-
-			InRHICmdList.UnlockVertexBuffer(VertexBufferRHI);
-		});
-
-		RHICmdList.RHIThreadFence(true);
+		ALLOC_COMMAND_CL(RHICmdList, FSlateUpdateInstanceBufferCommand)(InstanceBufferResource, RenderThreadBufferData);
 	}
 }
 
