Comments out compilation failure in Slate from 7235503 that collided with 7039544

diff --git a/Engine/Source/Runtime/SlateRHIRenderer/Private/SlateElementIndexBuffer.cpp b/Engine/Source/Runtime/SlateRHIRenderer/Private/SlateElementIndexBuffer.cpp
--- a/Engine/Source/Runtime/SlateRHIRenderer/Private/SlateElementIndexBuffer.cpp
+++ b/Engine/Source/Runtime/SlateRHIRenderer/Private/SlateElementIndexBuffer.cpp
@@ -105,7 +105,9 @@ void* FSlateElementIndexBuffer::LockBuffer(int32 NumIndices, bool bInRenderThrea
 	}
 	else
 	{
-		return GDynamicRHI->RHILockIndexBuffer(IndexBufferRHI, 0, RequiredBufferSize, RLM_WriteOnly);
+		check(0); // TODO
+		return nullptr;
+		//return GDynamicRHI->RHILockIndexBuffer(IndexBufferRHI, 0, RequiredBufferSize, RLM_WriteOnly);
 	}
 }
 
@@ -117,7 +119,8 @@ void FSlateElementIndexBuffer::UnlockBuffer(bool bInRenderThread)
 	}
 	else
 	{
-		GDynamicRHI->RHIUnlockIndexBuffer(IndexBufferRHI);
+		check(0); // TODO
+		//GDynamicRHI->RHIUnlockIndexBuffer(IndexBufferRHI);
 	}
 }
 
diff --git a/Engine/Source/Runtime/SlateRHIRenderer/Private/SlateElementVertexBuffer.h b/Engine/Source/Runtime/SlateRHIRenderer/Private/SlateElementVertexBuffer.h
--- a/Engine/Source/Runtime/SlateRHIRenderer/Private/SlateElementVertexBuffer.h
+++ b/Engine/Source/Runtime/SlateRHIRenderer/Private/SlateElementVertexBuffer.h
@@ -129,7 +129,9 @@ public:
 		else
 		{
 			// RHI thread
-			return GDynamicRHI->RHILockVertexBuffer(VertexBufferRHI, 0, RequiredBufferSize, RLM_WriteOnly);
+			check(0); // TODO
+			return nullptr;
+			//return GDynamicRHI->RHILockVertexBuffer(VertexBufferRHI, 0, RequiredBufferSize, RLM_WriteOnly);
 		}
 	
 	}
@@ -142,7 +144,8 @@ public:
 		}
 		else
 		{
-			return GDynamicRHI->RHIUnlockVertexBuffer(VertexBufferRHI);
+			check(0); // TODO
+			//return GDynamicRHI->RHIUnlockVertexBuffer(VertexBufferRHI);
 		}
 	}
 
