Refactored FSlateUpdatableInstanceBuffer to remove the render thread flush in its destructor. The buffer now maintains its own proxy object, the lifetime of which is handled by the render thread, removing the need for a flush in the destructor. This also simplifies the update functions to simply taking a TArray of FVector4, which is typedef'ed as FSlateInstanceBufferData.

This change also resolves the remaining issues from the recent //UE4/Main -> //UE4/Dev-Rendering merge involving the RHI Lock/Unlock buffer API changes.

diff --git a/Engine/Source/Runtime/SlateCore/Private/Rendering/DrawElements.cpp b/Engine/Source/Runtime/SlateCore/Private/Rendering/DrawElements.cpp
--- a/Engine/Source/Runtime/SlateCore/Private/Rendering/DrawElements.cpp
+++ b/Engine/Source/Runtime/SlateCore/Private/Rendering/DrawElements.cpp
@@ -554,7 +554,7 @@ void FSlateDrawElement::MakeCustomVerts(FSlateWindowElementList& ElementList, ui
 	FSlateCustomVertsPayload& DataPayload = ElementList.CreatePayload<FSlateCustomVertsPayload>(Element);
 
 	const FSlateShaderResourceProxy* RenderingProxy = InRenderResourceHandle.GetResourceProxy();
-	DataPayload.SetCustomVerts(RenderingProxy, InVerts, InIndexes, InInstanceData, InInstanceOffset, InNumInstances);
+	DataPayload.SetCustomVerts(RenderingProxy, InVerts, InIndexes, InInstanceData->GetRenderProxy(), InInstanceOffset, InNumInstances);
 
 	Element.Init(ElementList, EElementType::ET_CustomVerts, InLayer, FPaintGeometry(), InDrawEffects);
 	Element.RenderTransform = FSlateRenderTransform();
diff --git a/Engine/Source/Runtime/SlateCore/Public/Rendering/DrawElementPayloads.h b/Engine/Source/Runtime/SlateCore/Public/Rendering/DrawElementPayloads.h
--- a/Engine/Source/Runtime/SlateCore/Public/Rendering/DrawElementPayloads.h
+++ b/Engine/Source/Runtime/SlateCore/Public/Rendering/DrawElementPayloads.h
@@ -302,11 +302,11 @@ struct FSlateCustomVertsPayload : public FSlateDataPayload
 	TArray<SlateIndex> Indices;
 
 	// Instancing support
-	ISlateUpdatableInstanceBuffer* InstanceData;
+	ISlateUpdatableInstanceBufferRenderProxy* InstanceData;
 	uint32 InstanceOffset;
 	uint32 NumInstances;
 
-	void SetCustomVerts(const FSlateShaderResourceProxy* InRenderProxy, const TArray<FSlateVertex>& InVerts, const TArray<SlateIndex>& InIndices, ISlateUpdatableInstanceBuffer* InInstanceData, uint32 InInstanceOffset, uint32 InNumInstances)
+	void SetCustomVerts(const FSlateShaderResourceProxy* InRenderProxy, const TArray<FSlateVertex>& InVerts, const TArray<SlateIndex>& InIndices, ISlateUpdatableInstanceBufferRenderProxy* InInstanceData, uint32 InInstanceOffset, uint32 InNumInstances)
 	{
 		ResourceProxy = InRenderProxy;
 
diff --git a/Engine/Source/Runtime/SlateCore/Public/Rendering/RenderingCommon.h b/Engine/Source/Runtime/SlateCore/Public/Rendering/RenderingCommon.h
--- a/Engine/Source/Runtime/SlateCore/Public/Rendering/RenderingCommon.h
+++ b/Engine/Source/Runtime/SlateCore/Public/Rendering/RenderingCommon.h
@@ -882,80 +882,35 @@ public:
 	virtual void DrawRenderThread(class FRHICommandListImmediate& RHICmdList, const void* RenderTarget) = 0;
 };
 
+/*
+ * A proxy object used to access a slate per-instance data buffer on the render thread.
+ */
+class ISlateUpdatableInstanceBufferRenderProxy
+{
+protected:
+	virtual ~ISlateUpdatableInstanceBufferRenderProxy() {};
+
+public:
+	virtual void BindStreamSource(class FRHICommandList& RHICmdList, int32 StreamIndex, uint32 InstanceOffset) = 0;
+};
+
+typedef TArray<FVector4> FSlateInstanceBufferData;
+
 /**
  * Represents a per instance data buffer for a custom Slate mesh element.
- * Use FSlateInstanceBufferUpdate to update the per-instance data.
- *  e.g.
- *    TSharedRef<FSlateInstanceBufferUpdate> NewUpdate = InstanceBuffer.BeginUpdate();
- *     NewUpdate.GetData().Add( FVector4(1,1,1,1) )
- *     FSlateInstanceBufferUpdate::CommitUpdate(NewUpdate);
  */
 class ISlateUpdatableInstanceBuffer
 {
 public:
-	virtual ~ISlateUpdatableInstanceBuffer(){};
+	virtual ~ISlateUpdatableInstanceBuffer() {};
 	friend class FSlateInstanceBufferUpdate;
 
-   /**
-	* Use this method to begin a new update to this instance of the buffer:
-	*/
-	virtual TSharedPtr<class FSlateInstanceBufferUpdate> BeginUpdate() = 0;
-
 	/** How many instances should we draw? */
 	virtual uint32 GetNumInstances() const = 0;
 
-private:
-	friend class FSlateInstanceBufferUpdate;
-
-	/** Updates rendering data for the GPU */
-	virtual void UpdateRenderingData(int32 NumInstancesToUse) = 0;
-
-	/** @return an array of instance data that is safe to populate (e.g not in use by the renderer) */
-	virtual TArray<FVector4>& GetBufferData() = 0;
-};
-
-/** Represents an update to the per-instance buffer. */
-class FSlateInstanceBufferUpdate
-{
-public:
-	/** Access the per-instance data for modiciation */
-	FORCEINLINE TArray<FVector4>& GetData(){ return Data; }
-	
-	/** Send an update to the render thread */
-	static void CommitUpdate(TSharedPtr<FSlateInstanceBufferUpdate>& UpdateToCommit)
-	{
-		ensure(UpdateToCommit.GetSharedReferenceCount() == 1);
-		UpdateToCommit->CommitUpdate_Internal();
-		UpdateToCommit.Reset();
-	}
-
-	~FSlateInstanceBufferUpdate()
-	{
-		if (!bWasCommitted)
-		{
-			CommitUpdate_Internal();
-		}
-	}
-
-private:
-	friend class FSlateUpdatableInstanceBuffer;
-	FSlateInstanceBufferUpdate(ISlateUpdatableInstanceBuffer& InBuffer)
-		: Buffer(InBuffer)
-		, Data(InBuffer.GetBufferData())
-		, InstanceCount(0)
-		, bWasCommitted(false)
-	{
-	}
-
-	void CommitUpdate_Internal()
-	{
-		Buffer.UpdateRenderingData(Data.Num());
-		bWasCommitted = true;
-	}
+	/** Returns the pointer to the render proxy, to be forwarded to the render thread. */
+	virtual ISlateUpdatableInstanceBufferRenderProxy* GetRenderProxy() const = 0;
 
-	ISlateUpdatableInstanceBuffer& Buffer;
-	TArray<FVector4>& Data;
-	uint32 InstanceCount;
-	bool bWasCommitted;
+	/** Updates the buffer with the provided data. */
+	virtual void Update(FSlateInstanceBufferData& Data) = 0;
 };
-
diff --git a/Engine/Source/Runtime/SlateCore/Public/Rendering/SlateRenderBatch.h b/Engine/Source/Runtime/SlateCore/Public/Rendering/SlateRenderBatch.h
--- a/Engine/Source/Runtime/SlateCore/Public/Rendering/SlateRenderBatch.h
+++ b/Engine/Source/Runtime/SlateCore/Public/Rendering/SlateRenderBatch.h
@@ -141,7 +141,7 @@ public:
 	/** Shader Resource to use with this batch.  */
 	const FSlateShaderResource* ShaderResource;
 
-	ISlateUpdatableInstanceBuffer* InstanceData;
+	ISlateUpdatableInstanceBufferRenderProxy* InstanceData;
 
 	// Source Data
 	FSlateVertexArray* SourceVertices;
diff --git a/Engine/Source/Runtime/SlateRHIRenderer/Private/SlateElementIndexBuffer.cpp b/Engine/Source/Runtime/SlateRHIRenderer/Private/SlateElementIndexBuffer.cpp
--- a/Engine/Source/Runtime/SlateRHIRenderer/Private/SlateElementIndexBuffer.cpp
+++ b/Engine/Source/Runtime/SlateRHIRenderer/Private/SlateElementIndexBuffer.cpp
@@ -96,34 +96,6 @@ void FSlateElementIndexBuffer::PreFillBuffer(int32 RequiredIndexCount, bool bShr
 	}
 }
 
-void* FSlateElementIndexBuffer::LockBuffer(int32 NumIndices, bool bInRenderThread)
-{
-	uint32 RequiredBufferSize = NumIndices*sizeof(SlateIndex);		
-	if(bInRenderThread)
-	{
-		return RHILockIndexBuffer(IndexBufferRHI, 0, RequiredBufferSize, RLM_WriteOnly);
-	}
-	else
-	{
-		check(0); // TODO
-		return nullptr;
-		//return GDynamicRHI->RHILockIndexBuffer(IndexBufferRHI, 0, RequiredBufferSize, RLM_WriteOnly);
-	}
-}
-
-void FSlateElementIndexBuffer::UnlockBuffer(bool bInRenderThread)
-{
-	if (bInRenderThread)
-	{
-		RHIUnlockIndexBuffer(IndexBufferRHI);
-	}
-	else
-	{
-		check(0); // TODO
-		//GDynamicRHI->RHIUnlockIndexBuffer(IndexBufferRHI);
-	}
-}
-
 /** Releases the index buffers RHI resource. */
 void FSlateElementIndexBuffer::ReleaseDynamicRHI()
 {
diff --git a/Engine/Source/Runtime/SlateRHIRenderer/Private/SlateElementIndexBuffer.h b/Engine/Source/Runtime/SlateRHIRenderer/Private/SlateElementIndexBuffer.h
--- a/Engine/Source/Runtime/SlateRHIRenderer/Private/SlateElementIndexBuffer.h
+++ b/Engine/Source/Runtime/SlateRHIRenderer/Private/SlateElementIndexBuffer.h
@@ -40,9 +40,6 @@ public:
 	/** Resizes buffer, accumulates states safely on render thread */
 	void PreFillBuffer(int32 RequiredIndexCount, bool bShrinkToMinSize);
 
-	void* LockBuffer(int32 NumIndices, bool bInRenderThread);
-	void UnlockBuffer(bool bInRenderThread);
-
 	int32 GetMinBufferSize() const { return MinBufferSize; }
 
 private:
diff --git a/Engine/Source/Runtime/SlateRHIRenderer/Private/SlateElementVertexBuffer.h b/Engine/Source/Runtime/SlateRHIRenderer/Private/SlateElementVertexBuffer.h
--- a/Engine/Source/Runtime/SlateRHIRenderer/Private/SlateElementVertexBuffer.h
+++ b/Engine/Source/Runtime/SlateRHIRenderer/Private/SlateElementVertexBuffer.h
@@ -117,38 +117,6 @@ public:
 
 	int32 GetMinBufferSize() const { return MinBufferSize; }
 
-	void* LockBuffer(int32 NumVertices, bool bInRenderThread)
-	{
-		uint32 RequiredBufferSize = NumVertices*sizeof(VertexType);	
-
-		if (bInRenderThread)
-		{
-			// Render thread
-			return RHILockVertexBuffer(VertexBufferRHI, 0, RequiredBufferSize, RLM_WriteOnly);
-		}
-		else
-		{
-			// RHI thread
-			check(0); // TODO
-			return nullptr;
-			//return GDynamicRHI->RHILockVertexBuffer(VertexBufferRHI, 0, RequiredBufferSize, RLM_WriteOnly);
-		}
-	
-	}
-
-	void UnlockBuffer(bool bInRenderThread)
-	{
-		if (bInRenderThread)
-		{
-			RHIUnlockVertexBuffer(VertexBufferRHI);
-		}
-		else
-		{
-			check(0); // TODO
-			//return GDynamicRHI->RHIUnlockVertexBuffer(VertexBufferRHI);
-		}
-	}
-
 private:
 	/** Resizes the buffer to the passed in size.  Preserves internal data*/
 	void ResizeBuffer( int32 NewSizeBytes )
diff --git a/Engine/Source/Runtime/SlateRHIRenderer/Private/SlateRHIRenderingPolicy.cpp b/Engine/Source/Runtime/SlateRHIRenderer/Private/SlateRHIRenderingPolicy.cpp
--- a/Engine/Source/Runtime/SlateRHIRenderer/Private/SlateRHIRenderingPolicy.cpp
+++ b/Engine/Source/Runtime/SlateRHIRenderer/Private/SlateRHIRenderingPolicy.cpp
@@ -131,58 +131,6 @@ void FSlateRHIRenderingPolicy::EndDrawingWindows()
 	check( IsInParallelRenderingThread() );
 }
 
-struct FSlateUpdateVertexAndIndexBuffers final : public FRHICommand<FSlateUpdateVertexAndIndexBuffers>
-{
-	TSlateElementVertexBuffer<FSlateVertex>& VertexBuffer;
-	FSlateElementIndexBuffer& IndexBuffer;
-	FSlateBatchData& BatchData;
-	bool bAbsoluteIndices;
-
-	FSlateUpdateVertexAndIndexBuffers(TSlateElementVertexBuffer<FSlateVertex>& InVertexBuffer, FSlateElementIndexBuffer& InIndexBuffer, FSlateBatchData& InBatchData, bool bInAbsoluteIndices)
-		: VertexBuffer(InVertexBuffer)
-		, IndexBuffer(InIndexBuffer)
-		, BatchData(InBatchData)
-		, bAbsoluteIndices(bInAbsoluteIndices)
-	{
-		check(IsInRenderingThread());
-	}
-
-	void Execute(FRHICommandListBase& CmdList)
-	{
-		SCOPE_CYCLE_COUNTER(STAT_SlateUpdateBufferRTTime);
-
-		const bool bIsInRenderingThread = !IsRunningRHIInSeparateThread() || CmdList.Bypass();
-
-		const FSlateVertexArray& FinalVertexData = BatchData.GetFinalVertexData();
-		const FSlateIndexArray& FinalIndexData = BatchData.GetFinalIndexData();
-
-		const int32 NumBatchedVertices = FinalVertexData.Num();
-		const int32 NumBatchedIndices = FinalIndexData.Num();
-		//int32 RequiredVertexBufferSize = NumBatchedVertices*sizeof(FSlateVertex);
-		uint8* VertexBufferData = (uint8*)VertexBuffer.LockBuffer(NumBatchedVertices, bIsInRenderingThread);
-
-		//uint32 RequiredIndexBufferSize = NumBatchedIndices*sizeof(SlateIndex);		
-		uint8* IndexBufferData = (uint8*)IndexBuffer.LockBuffer(NumBatchedIndices, bIsInRenderingThread);
-
-		//Early out if we have an invalid buffer (might have lost context and now have invalid buffers)
-		if ((nullptr != VertexBufferData) && (nullptr != IndexBufferData))
-		{
-			FMemory::Memcpy(VertexBufferData, FinalVertexData.GetData(), FinalVertexData.Num() * sizeof(FSlateVertex));
-			FMemory::Memcpy(IndexBufferData, FinalIndexData.GetData(), FinalIndexData.Num() * sizeof(SlateIndex));
-		}
-
-		if (nullptr != VertexBufferData)
-		{
-			VertexBuffer.UnlockBuffer(bIsInRenderingThread);
-		}
-
-		if (nullptr != IndexBufferData)
-		{
-			IndexBuffer.UnlockBuffer(bIsInRenderingThread);
-		}
-	}
-};
-
 void FSlateRHIRenderingPolicy::BuildRenderingBuffers(FRHICommandListImmediate& RHICmdList, FSlateBatchData& InBatchData)
 {
 	SCOPE_CYCLE_COUNTER(STAT_SlateUpdateBufferRTTime);
@@ -207,16 +155,35 @@ void FSlateRHIRenderingPolicy::BuildRenderingBuffers(FRHICommandListImmediate& R
 		MasterVertexBuffer.PreFillBuffer(NumVertices, bShouldShrinkResources);
 		MasterIndexBuffer.PreFillBuffer(NumIndices, bShouldShrinkResources);
 
-		if (!IsRunningRHIInSeparateThread() || RHICmdList.Bypass())
+		RHICmdList.EnqueueLambda([
+			VertexBuffer = MasterVertexBuffer.VertexBufferRHI.GetReference(),
+			IndexBuffer = MasterIndexBuffer.IndexBufferRHI.GetReference(),
+			&InBatchData,
+			bAbsoluteIndices
+		](FRHICommandListImmediate& RHICmdList)
 		{
-			FSlateUpdateVertexAndIndexBuffers UpdateBufferCommand(MasterVertexBuffer, MasterIndexBuffer, InBatchData, bAbsoluteIndices);
+			SCOPE_CYCLE_COUNTER(STAT_SlateUpdateBufferRTTime);
 
-			UpdateBufferCommand.Execute(RHICmdList);
-		}
-		else
-		{
-			ALLOC_COMMAND_CL(RHICmdList, FSlateUpdateVertexAndIndexBuffers)(MasterVertexBuffer, MasterIndexBuffer, InBatchData, bAbsoluteIndices);
-		}
+			const FSlateVertexArray& FinalVertexData = InBatchData.GetFinalVertexData();
+			const FSlateIndexArray& FinalIndexData = InBatchData.GetFinalIndexData();
+
+			const int32 NumBatchedVertices = FinalVertexData.Num();
+			const int32 NumBatchedIndices = FinalIndexData.Num();
+
+			uint32 RequiredVertexBufferSize = NumBatchedVertices * sizeof(FSlateVertex);
+			uint8* VertexBufferData = (uint8*)RHICmdList.LockVertexBuffer(VertexBuffer, 0, RequiredVertexBufferSize, RLM_WriteOnly);
+
+			uint32 RequiredIndexBufferSize = NumBatchedIndices * sizeof(SlateIndex);
+			uint8* IndexBufferData = (uint8*)RHICmdList.LockIndexBuffer(IndexBuffer, 0, RequiredIndexBufferSize, RLM_WriteOnly);
+
+			FMemory::Memcpy(VertexBufferData, FinalVertexData.GetData(), RequiredVertexBufferSize);
+			FMemory::Memcpy(IndexBufferData, FinalIndexData.GetData(), RequiredIndexBufferSize);
+
+			RHICmdList.UnlockVertexBuffer(VertexBuffer);
+			RHICmdList.UnlockIndexBuffer(IndexBuffer);
+		});
+
+		RHICmdList.RHIThreadFence(true);
 	}
 
 	checkSlow(MasterVertexBuffer.GetBufferUsageSize() <= MasterVertexBuffer.GetBufferSize());
@@ -1130,8 +1097,7 @@ void FSlateRHIRenderingPolicy::DrawElements(
 
 								if (GRHISupportsInstancing)
 								{
-									FSlateUpdatableInstanceBuffer* InstanceBuffer = (FSlateUpdatableInstanceBuffer*)RenderBatch.InstanceData;
-									InstanceBuffer->BindStreamSource(RHICmdList, 1, RenderBatch.InstanceOffset);
+									RenderBatch.InstanceData->BindStreamSource(RHICmdList, 1, RenderBatch.InstanceOffset);
 
 									// for RHIs that can't handle VertexOffset, we need to offset the stream source each time
 									if (1 || (!GRHISupportsBaseVertexIndex && !bAbsoluteIndices))
diff --git a/Engine/Source/Runtime/SlateRHIRenderer/Private/SlateUpdatableBuffer.cpp b/Engine/Source/Runtime/SlateRHIRenderer/Private/SlateUpdatableBuffer.cpp
--- a/Engine/Source/Runtime/SlateRHIRenderer/Private/SlateUpdatableBuffer.cpp
+++ b/Engine/Source/Runtime/SlateRHIRenderer/Private/SlateUpdatableBuffer.cpp
@@ -3,102 +3,66 @@
 #include "SlateUpdatableBuffer.h"
 #include "RenderingThread.h"
 
-DECLARE_CYCLE_STAT(TEXT("UpdateInstanceBuffer Time"), STAT_SlateUpdateInstanceBuffer, STATGROUP_Slate);
+DECLARE_CYCLE_STAT(TEXT("UpdateInstanceBuffer Time (RT)"), STAT_SlateUpdateInstanceBuffer_RT, STATGROUP_Slate);
+DECLARE_CYCLE_STAT(TEXT("UpdateInstanceBuffer Time (RHIT)"), STAT_SlateUpdateInstanceBuffer_RHIT, STATGROUP_Slate);
 
-struct FSlateUpdateInstanceBufferCommand final : public FRHICommand<FSlateUpdateInstanceBufferCommand>
+FSlateUpdatableInstanceBuffer::FSlateUpdatableInstanceBuffer(int32 InitialInstanceCount)
 {
-	TSlateElementVertexBuffer<FVector4>& InstanceBuffer;
-	const TArray<FVector4>& InstanceData;
-
-	FSlateUpdateInstanceBufferCommand(TSlateElementVertexBuffer<FVector4>& InInstanceBuffer, const TArray<FVector4>& InInstanceData )
-		: InstanceBuffer(InInstanceBuffer)
-		, InstanceData(InInstanceData)
-	{}
-
-	void Execute(FRHICommandListBase& CmdList)
-	{
-		SCOPE_CYCLE_COUNTER( STAT_SlateUpdateInstanceBuffer );
-		const bool bIsInRenderingThread = !IsRunningRHIInSeparateThread() || CmdList.Bypass();
-
-		int32 RequiredVertexBufferSize = InstanceData.Num()*sizeof(FVector4);
-		uint8* InstanceBufferData = (uint8*)InstanceBuffer.LockBuffer(RequiredVertexBufferSize, bIsInRenderingThread);
-
-		FMemory::Memcpy( InstanceBufferData, InstanceData.GetData(), InstanceData.Num()*sizeof(FVector4) );
-	
-		InstanceBuffer.UnlockBuffer(bIsInRenderingThread);
-	}
-};
-
-FSlateUpdatableInstanceBuffer::FSlateUpdatableInstanceBuffer( int32 InitialInstanceCount )
-	: FreeBufferIndex(0)
-{
-	InstanceBufferResource.Init(InitialInstanceCount);
-	for( int32 BufferIndex = 0; BufferIndex < SlateRHIConstants::NumBuffers; ++BufferIndex )
-	{
-		BufferData[BufferIndex].Reserve( InitialInstanceCount );
-	}
-
+	Proxy = new FRenderProxy;
+	Proxy->InstanceBufferResource.Init(InitialInstanceCount);
 }
 
 FSlateUpdatableInstanceBuffer::~FSlateUpdatableInstanceBuffer()
 {
-	InstanceBufferResource.Destroy();
-	FlushRenderingCommands();
-}
+	ENQUEUE_RENDER_COMMAND(SlateUpdatableInstanceBuffer_DeleteProxy)(
+		[Ptr = Proxy](FRHICommandListImmediate&)
+	{
+		delete Ptr;
+	});
 
-void FSlateUpdatableInstanceBuffer::BindStreamSource(FRHICommandListImmediate& RHICmdList, int32 StreamIndex, uint32 InstanceOffset)
-{
-	RHICmdList.SetStreamSource(StreamIndex, InstanceBufferResource.VertexBufferRHI, InstanceOffset*sizeof(FVector4));
+	Proxy = nullptr;
 }
 
-TSharedPtr<class FSlateInstanceBufferUpdate> FSlateUpdatableInstanceBuffer::BeginUpdate()
+void FSlateUpdatableInstanceBuffer::Update(FSlateInstanceBufferData& Data)
 {
-	return MakeShareable(new FSlateInstanceBufferUpdate(*this));
-}
+	check(!IsInRenderingThread() && !IsInRHIThread());
 
-uint32 FSlateUpdatableInstanceBuffer::GetNumInstances() const
-{
-	return NumInstances;
-}
-
-void FSlateUpdatableInstanceBuffer::UpdateRenderingData(int32 NumInstancesToUse)
-{
-	NumInstances = NumInstancesToUse;
+	NumInstances = Data.Num();
 	if (NumInstances > 0)
 	{
-		// Enqueue a command to unlock the draw buffer after all windows have been drawn
-		FSlateUpdatableInstanceBuffer* Self = this;
-		int32 BufferIndex = FreeBufferIndex;
-		ENQUEUE_RENDER_COMMAND(SlateBeginDrawingWindowsCommand)(
-			[Self, BufferIndex](FRHICommandListImmediate& RHICmdList)
-			{
-				Self->UpdateRenderingData_RenderThread(RHICmdList, BufferIndex);
-			});
-	
-		FreeBufferIndex = (FreeBufferIndex + 1) % SlateRHIConstants::NumBuffers;
+		// Enqueue a render thread command to update the proxy with the new data.
+		ENQUEUE_RENDER_COMMAND(SlateUpdatableInstanceBuffer_Update)(
+			[Ptr = Proxy, LocalDataRT = MoveTemp(Data)](FRHICommandListImmediate& RHICmdList) mutable
+		{
+			Ptr->Update(RHICmdList, LocalDataRT);
+		});
 	}
 }
 
-TArray<FVector4>& FSlateUpdatableInstanceBuffer::GetBufferData()
-{
-	return BufferData[FreeBufferIndex];
-}
-
-void FSlateUpdatableInstanceBuffer::UpdateRenderingData_RenderThread(FRHICommandListImmediate& RHICmdList, int32 BufferIndex)
+void FSlateUpdatableInstanceBuffer::FRenderProxy::Update(FRHICommandListImmediate& RHICmdList, FSlateInstanceBufferData& Data)
 {
-	SCOPE_CYCLE_COUNTER( STAT_SlateUpdateInstanceBuffer );
+	SCOPE_CYCLE_COUNTER(STAT_SlateUpdateInstanceBuffer_RT);
 
-	const TArray<FVector4>& RenderThreadBufferData = BufferData[BufferIndex];
-	InstanceBufferResource.PreFillBuffer( RenderThreadBufferData.Num(), false );
+	InstanceBufferResource.PreFillBuffer(Data.Num(), false);
 
-	if(!IsRunningRHIInSeparateThread() || RHICmdList.Bypass())
-	{
-		FSlateUpdateInstanceBufferCommand Command(InstanceBufferResource, RenderThreadBufferData);
-		Command.Execute(RHICmdList);
-	}
-	else
+	// Enqueue the lock/unlock to the RHI thread
+	FRHIVertexBuffer* VertexBuffer = InstanceBufferResource.VertexBufferRHI;
+	RHICmdList.EnqueueLambda([VertexBuffer, LocalData = MoveTemp(Data)](FRHICommandListImmediate& InRHICmdList)
 	{
-		ALLOC_COMMAND_CL(RHICmdList, FSlateUpdateInstanceBufferCommand)(InstanceBufferResource, RenderThreadBufferData);
-	}
+		SCOPE_CYCLE_COUNTER(STAT_SlateUpdateInstanceBuffer_RHIT);
+
+		int32 RequiredVertexBufferSize = LocalData.Num() * LocalData.GetTypeSize();
+		uint8* InstanceBufferData = (uint8*)InRHICmdList.LockVertexBuffer(VertexBuffer, 0, RequiredVertexBufferSize, RLM_WriteOnly);
+
+		FMemory::Memcpy(InstanceBufferData, LocalData.GetData(), RequiredVertexBufferSize);
+
+		InRHICmdList.UnlockVertexBuffer(VertexBuffer);
+	});
+
+	RHICmdList.RHIThreadFence(true);
 }
 
+void FSlateUpdatableInstanceBuffer::FRenderProxy::BindStreamSource(FRHICommandList& RHICmdList, int32 StreamIndex, uint32 InstanceOffset)
+{
+	RHICmdList.SetStreamSource(StreamIndex, InstanceBufferResource.VertexBufferRHI, InstanceOffset * sizeof(FVector4));
+}
diff --git a/Engine/Source/Runtime/SlateRHIRenderer/Private/SlateUpdatableBuffer.h b/Engine/Source/Runtime/SlateRHIRenderer/Private/SlateUpdatableBuffer.h
--- a/Engine/Source/Runtime/SlateRHIRenderer/Private/SlateUpdatableBuffer.h
+++ b/Engine/Source/Runtime/SlateRHIRenderer/Private/SlateUpdatableBuffer.h
@@ -10,29 +10,34 @@
 /**
  * Represents a per instance data buffer for a custom Slate mesh element
  */
-class FSlateUpdatableInstanceBuffer : public ISlateUpdatableInstanceBuffer
+class FSlateUpdatableInstanceBuffer final : public ISlateUpdatableInstanceBuffer
 {
-public:
+	// Owned by the render thread
+	struct FRenderProxy final : public ISlateUpdatableInstanceBufferRenderProxy
+	{
+		TSlateElementVertexBuffer<FVector4> InstanceBufferResource;
+
+		virtual ~FRenderProxy()
+		{
+			InstanceBufferResource.Destroy();
+		}
+
+		void Update(FRHICommandListImmediate& RHICmdList, FSlateInstanceBufferData& Data);
+
+		virtual void BindStreamSource(FRHICommandList& RHICmdList, int32 StreamIndex, uint32 InstanceOffset) override final;
+	} *Proxy;
 
-	FSlateUpdatableInstanceBuffer( int32 InstanceCount );
+public:
+	FSlateUpdatableInstanceBuffer(int32 InstanceCount);
 	~FSlateUpdatableInstanceBuffer();
-	void BindStreamSource(FRHICommandListImmediate& RHICmdList, int32 StreamIndex, uint32 InstanceOffset);
 
 private:
 	// BEGIN ISlateUpdatableInstanceBuffer
-	virtual TSharedPtr<class FSlateInstanceBufferUpdate> BeginUpdate() override;
-	virtual uint32 GetNumInstances() const override;
-	virtual void UpdateRenderingData(int32 NumInstancesToUse) override;
-	virtual TArray<FVector4>& GetBufferData() override;
+	virtual uint32 GetNumInstances() const override final { return NumInstances; }
+	virtual ISlateUpdatableInstanceBufferRenderProxy* GetRenderProxy() const override final { return Proxy; }
+	virtual void Update(FSlateInstanceBufferData& Data) override final;
 	// END ISlateUpdatableInstanceBuffer
 
-	void UpdateRenderingData_RenderThread(FRHICommandListImmediate& RHICmdList, int32 BufferIndex);
 private:
-	TArray<FVector4> BufferData[SlateRHIConstants::NumBuffers];
-	TSlateElementVertexBuffer<FVector4> InstanceBufferResource;
 	uint32 NumInstances;
-	int32 FreeBufferIndex;
 };
-
-
-
diff --git a/Engine/Source/Runtime/UMG/Private/Slate/SMeshWidget.cpp b/Engine/Source/Runtime/UMG/Private/Slate/SMeshWidget.cpp
--- a/Engine/Source/Runtime/UMG/Private/Slate/SMeshWidget.cpp
+++ b/Engine/Source/Runtime/UMG/Private/Slate/SMeshWidget.cpp
@@ -139,15 +139,10 @@ void SMeshWidget::EnableInstancing(uint32 MeshId, int32 InitialSize)
 	}
 }
 
-TSharedPtr<FSlateInstanceBufferUpdate> SMeshWidget::BeginPerInstanceBufferUpdate(uint32 MeshId, int32 InitialSize)
+void SMeshWidget::UpdatePerInstanceBuffer(uint32 MeshId, FSlateInstanceBufferData& Data)
 {
-	EnableInstancing(MeshId, InitialSize);
-	return BeginPerInstanceBufferUpdateConst( MeshId );
-}
-
-TSharedPtr<FSlateInstanceBufferUpdate> SMeshWidget::BeginPerInstanceBufferUpdateConst(uint32 MeshId) const
-{
-	return RenderData[MeshId].PerInstanceBuffer->BeginUpdate();
+	EnableInstancing(MeshId, Data.Num() * Data.GetTypeSize());
+	RenderData[MeshId].PerInstanceBuffer->Update(Data);
 }
 
 
@@ -249,12 +244,12 @@ FString SMeshWidget::GetReferencerName() const
 	return TEXT("SMeshWidget");
 }
 
-void SMeshWidget::PushUpdate(uint32 VectorArtId, const SMeshWidget& Widget, const FVector2D& Position, float Scale, uint32 BaseAddress)
+void SMeshWidget::PushUpdate(uint32 VectorArtId, SMeshWidget& Widget, const FVector2D& Position, float Scale, uint32 BaseAddress)
 {
 	PushUpdate(VectorArtId, Widget, Position, Scale, static_cast<float>(BaseAddress));
 }
 
-void SMeshWidget::PushUpdate(uint32 VectorArtId, const SMeshWidget& Widget, const FVector2D& Position, float Scale, float OptionalFloat /*= 0*/)
+void SMeshWidget::PushUpdate(uint32 VectorArtId, SMeshWidget& Widget, const FVector2D& Position, float Scale, float OptionalFloat /*= 0*/)
 {
 	FSlateVectorArtInstanceData Data;
 	Data.SetPosition(Position);
@@ -262,9 +257,8 @@ void SMeshWidget::PushUpdate(uint32 VectorArtId, const SMeshWidget& Widget, cons
 	Data.SetBaseAddress(OptionalFloat);
 
 	{
-		TSharedPtr<FSlateInstanceBufferUpdate> PerInstaceUpdate = Widget.BeginPerInstanceBufferUpdateConst(VectorArtId);
-		PerInstaceUpdate->GetData().Empty();
-		PerInstaceUpdate->GetData().Add(Data.GetData());
-		FSlateInstanceBufferUpdate::CommitUpdate(PerInstaceUpdate);
+		FSlateInstanceBufferData PerInstanceData;
+		PerInstanceData.Add(Data.GetData());
+		Widget.UpdatePerInstanceBuffer(VectorArtId, PerInstanceData);
 	}
 }
diff --git a/Engine/Source/Runtime/UMG/Public/Slate/SMeshWidget.h b/Engine/Source/Runtime/UMG/Public/Slate/SMeshWidget.h
--- a/Engine/Source/Runtime/UMG/Public/Slate/SMeshWidget.h
+++ b/Engine/Source/Runtime/UMG/Public/Slate/SMeshWidget.h
@@ -69,10 +69,9 @@ public:
 
 	/** Enable hardware instancing */
 	void EnableInstancing(uint32 MeshId, int32 InitialSize);
-	/** Begin an update to the per instance buffer. Automatically enables hardware instancing. */
-	TSharedPtr<FSlateInstanceBufferUpdate> BeginPerInstanceBufferUpdate(uint32 MeshId, int32 InitialSize);
-	/** Begin an update to the per instance buffer. Use */
-	TSharedPtr<FSlateInstanceBufferUpdate> BeginPerInstanceBufferUpdateConst(uint32 MeshId) const;
+
+	/** Updates the per instance buffer. Automatically enables hardware instancing. */
+	void UpdatePerInstanceBuffer(uint32 MeshId, FSlateInstanceBufferData& Data);
 
 protected:
 	// BEGIN SLeafWidget interface
@@ -86,8 +85,8 @@ protected:
 	// ~ FGCObject
 
 protected:
-	static void PushUpdate(uint32 VectorArtId, const SMeshWidget& Widget, const FVector2D& Position, float Scale, uint32 BaseAddress);
-	static void PushUpdate(uint32 VectorArtId, const SMeshWidget& Widget, const FVector2D& Position, float Scale, float OptionalFloat = 0);
+	static void PushUpdate(uint32 VectorArtId, SMeshWidget& Widget, const FVector2D& Position, float Scale, uint32 BaseAddress);
+	static void PushUpdate(uint32 VectorArtId, SMeshWidget& Widget, const FVector2D& Position, float Scale, float OptionalFloat = 0);
 
 	struct FRenderData
 	{
