Fixed clang warnings (-Wshadow) that were removed due to merge conflicts.

#rnx

diff --git a/Engine/Source/Runtime/SlateRHIRenderer/Private/SlateRHIRenderingPolicy.cpp b/Engine/Source/Runtime/SlateRHIRenderer/Private/SlateRHIRenderingPolicy.cpp
--- a/Engine/Source/Runtime/SlateRHIRenderer/Private/SlateRHIRenderingPolicy.cpp
+++ b/Engine/Source/Runtime/SlateRHIRenderer/Private/SlateRHIRenderingPolicy.cpp
@@ -30,6 +30,7 @@
 extern void UpdateNoiseTextureParameters(FViewUniformShaderParameters& ViewUniformShaderParameters);
 
 DECLARE_CYCLE_STAT(TEXT("Update Buffers RT"), STAT_SlateUpdateBufferRTTime, STATGROUP_Slate);
+DECLARE_CYCLE_STAT(TEXT("Update Buffers RT"), STAT_SlateUpdateBufferRTTimeLambda, STATGROUP_Slate);
 
 
 DECLARE_DWORD_COUNTER_STAT(TEXT("Num Layers"), STAT_SlateNumLayers, STATGROUP_Slate);
@@ -160,27 +161,28 @@ void FSlateRHIRenderingPolicy::BuildRenderingBuffers(FRHICommandListImmediate& R
 			IndexBuffer = MasterIndexBuffer.IndexBufferRHI.GetReference(),
 			&InBatchData,
 			bAbsoluteIndices
-		](FRHICommandListImmediate& RHICmdList)
+		](FRHICommandListImmediate& InRHICmdList)
 		{
-			SCOPE_CYCLE_COUNTER(STAT_SlateUpdateBufferRTTime);
+			SCOPE_CYCLE_COUNTER(STAT_SlateUpdateBufferRTTimeLambda);
 
-			const FSlateVertexArray& FinalVertexData = InBatchData.GetFinalVertexData();
-			const FSlateIndexArray& FinalIndexData = InBatchData.GetFinalIndexData();
+			// Note: Use "Lambda" prefix to prevent clang/gcc warnings of '-Wshadow' warning
+			const FSlateVertexArray& LambdaFinalVertexData = InBatchData.GetFinalVertexData();
+			const FSlateIndexArray& LambdaFinalIndexData = InBatchData.GetFinalIndexData();
 
-			const int32 NumBatchedVertices = FinalVertexData.Num();
-			const int32 NumBatchedIndices = FinalIndexData.Num();
+			const int32 NumBatchedVertices = LambdaFinalVertexData.Num();
+			const int32 NumBatchedIndices = LambdaFinalIndexData.Num();
 
 			uint32 RequiredVertexBufferSize = NumBatchedVertices * sizeof(FSlateVertex);
-			uint8* VertexBufferData = (uint8*)RHICmdList.LockVertexBuffer(VertexBuffer, 0, RequiredVertexBufferSize, RLM_WriteOnly);
+			uint8* VertexBufferData = (uint8*)InRHICmdList.LockVertexBuffer(VertexBuffer, 0, RequiredVertexBufferSize, RLM_WriteOnly);
 
 			uint32 RequiredIndexBufferSize = NumBatchedIndices * sizeof(SlateIndex);
-			uint8* IndexBufferData = (uint8*)RHICmdList.LockIndexBuffer(IndexBuffer, 0, RequiredIndexBufferSize, RLM_WriteOnly);
+			uint8* IndexBufferData = (uint8*)InRHICmdList.LockIndexBuffer(IndexBuffer, 0, RequiredIndexBufferSize, RLM_WriteOnly);
 
-			FMemory::Memcpy(VertexBufferData, FinalVertexData.GetData(), RequiredVertexBufferSize);
-			FMemory::Memcpy(IndexBufferData, FinalIndexData.GetData(), RequiredIndexBufferSize);
+			FMemory::Memcpy(VertexBufferData, LambdaFinalVertexData.GetData(), RequiredVertexBufferSize);
+			FMemory::Memcpy(IndexBufferData, LambdaFinalIndexData.GetData(), RequiredIndexBufferSize);
 
-			RHICmdList.UnlockVertexBuffer(VertexBuffer);
-			RHICmdList.UnlockIndexBuffer(IndexBuffer);
+			InRHICmdList.UnlockVertexBuffer(VertexBuffer);
+			InRHICmdList.UnlockIndexBuffer(IndexBuffer);
 		});
 
 		RHICmdList.RHIThreadFence(true);
