Fixed crash due to screenshot rect extending outside the backbuffer. Rect is now clamped to back buffer.

diff --git a/Engine/Source/Runtime/SlateRHIRenderer/Private/SlateRHIRenderer.cpp b/Engine/Source/Runtime/SlateRHIRenderer/Private/SlateRHIRenderer.cpp
--- a/Engine/Source/Runtime/SlateRHIRenderer/Private/SlateRHIRenderer.cpp
+++ b/Engine/Source/Runtime/SlateRHIRenderer/Private/SlateRHIRenderer.cpp
@@ -1078,7 +1078,27 @@ void FSlateRHIRenderer::DrawWindow_RenderThread(FRHICommandListImmediate& RHICmd
 	{
 		// take screenshot before swapbuffer
 		FTexture2DRHIRef BackBuffer = RHICmdList.GetViewportBackBuffer(ViewportInfo.ViewportRHI);
-		RHICmdList.ReadSurfaceData(BackBuffer, ScreenshotRect, *OutScreenshotData, FReadSurfaceDataFlags());
+
+		// Sanity check to make sure the user specified a valid screenshot rect.
+		FIntRect ClampedScreenshotRect;
+		ClampedScreenshotRect.Min = ScreenshotRect.Min;
+		ClampedScreenshotRect.Max = ScreenshotRect.Max.ComponentMin(BackBuffer->GetSizeXY());
+		ClampedScreenshotRect.Max = ScreenshotRect.Min.ComponentMax(ClampedScreenshotRect.Max);
+
+		if (ClampedScreenshotRect != ScreenshotRect)
+		{
+			UE_LOG(LogSlate, Warning, TEXT("Slate: Screenshot rect max coordinate had to be clamped from [%d, %d] to [%d, %d]"), ScreenshotRect.Max.X, ScreenshotRect.Max.Y, ClampedScreenshotRect.Max.X, ClampedScreenshotRect.Max.Y);
+		}
+
+		if (!ClampedScreenshotRect.IsEmpty())
+		{
+			RHICmdList.ReadSurfaceData(BackBuffer, ScreenshotRect, *OutScreenshotData, FReadSurfaceDataFlags());
+		}
+		else
+		{
+			UE_LOG(LogSlate, Warning, TEXT("Slate: Screenshot rect was empty! Skipping readback of back buffer."));
+		}
+
 		bTakingAScreenShot = false;
 		OutScreenshotData = NULL;
 	}
