Adding a change to make it so the editor still goes gamma correction even if the frame buffer format is float.

diff --git a/Engine/Source/Runtime/SlateRHIRenderer/Private/SlateRHIRenderingPolicy.cpp b/Engine/Source/Runtime/SlateRHIRenderer/Private/SlateRHIRenderingPolicy.cpp
--- a/Engine/Source/Runtime/SlateRHIRenderer/Private/SlateRHIRenderingPolicy.cpp
+++ b/Engine/Source/Runtime/SlateRHIRenderer/Private/SlateRHIRenderingPolicy.cpp
@@ -584,7 +584,7 @@ void FSlateRHIRenderingPolicy::DrawElements(
 	// and re-gammatization of the resulting blending operation in the 8/10 bit RGBA path.
 	// For Editor running in HDR then the gamma needs to be 2.2 and have a float back buffer format.
 
-	const float EngineGamma = ((BackBuffer.GetRenderTargetTexture()->GetFormat() == PF_FloatRGBA)&&(Params.bIsHDR==false)) ? 1.0 : GEngine ? GEngine->GetDisplayGamma() : 2.2f;
+	const float EngineGamma = (!GIsEditor && (BackBuffer.GetRenderTargetTexture()->GetFormat() == PF_FloatRGBA) && (Params.bIsHDR==false)) ? 1.0 : GEngine ? GEngine->GetDisplayGamma() : 2.2f;
 	const float DisplayGamma = bGammaCorrect ? EngineGamma : 1.0f;
 	const float DisplayContrast = GSlateContrast;
 
