Go back to the default behviour of trunc rather then ceil when converting window size from float to int. As this can put an extra pixel on the height or width

Originally it was just using a float which would be turned into an int (implicit trunc)
Then it was moved to round up to Ceil in SWindow
Then SlateRHIRenderer was move to Ceil to match SWindow as it was causing an extra resize
Now move both to use Trunc as thats what it was oringinally doing

diff --git a/Engine/Source/Runtime/SlateCore/Private/Widgets/SWindow.cpp b/Engine/Source/Runtime/SlateCore/Private/Widgets/SWindow.cpp
--- a/Engine/Source/Runtime/SlateCore/Private/Widgets/SWindow.cpp
+++ b/Engine/Source/Runtime/SlateCore/Private/Widgets/SWindow.cpp
@@ -882,7 +882,7 @@ void SWindow::ReshapeWindow( FVector2D NewPosition, FVector2D NewSize )
 	const FVector2D CurrentSize = GetSizeInScreen();
 
 	const FVector2D NewPositionTruncated = FVector2D(FMath::TruncToInt(NewPosition.X), FMath::TruncToInt(NewPosition.Y));
-	const FVector2D NewSizeRounded = FVector2D(FMath::CeilToInt(NewSize.X), FMath::CeilToInt(NewSize.Y));
+	const FVector2D NewSizeRounded = FVector2D(FMath::TruncToInt(NewSize.X), FMath::TruncToInt(NewSize.Y));
 
 	if ( CurrentPosition != NewPositionTruncated || CurrentSize != NewSizeRounded )
 	{
diff --git a/Engine/Source/Runtime/SlateRHIRenderer/Private/SlateRHIRenderer.cpp b/Engine/Source/Runtime/SlateRHIRenderer/Private/SlateRHIRenderer.cpp
--- a/Engine/Source/Runtime/SlateRHIRenderer/Private/SlateRHIRenderer.cpp
+++ b/Engine/Source/Runtime/SlateRHIRenderer/Private/SlateRHIRenderer.cpp
@@ -281,8 +281,8 @@ void FSlateRHIRenderer::CreateViewport(const TSharedRef<SWindow> Window)
 
 		// Clamp the window size to a reasonable default anything below 8 is a d3d warning and 8 is used anyway.
 		// @todo Slate: This is a hack to work around menus being summoned with 0,0 for window size until they are ticked.
-		int32 Width = FMath::Max(MIN_VIEWPORT_SIZE, FMath::CeilToInt(WindowSize.X));
-		int32 Height = FMath::Max(MIN_VIEWPORT_SIZE, FMath::CeilToInt(WindowSize.Y));
+		int32 Width = FMath::Max(MIN_VIEWPORT_SIZE, FMath::TruncToInt(WindowSize.X));
+		int32 Height = FMath::Max(MIN_VIEWPORT_SIZE, FMath::TruncToInt(WindowSize.Y));
 
 		// Sanity check dimensions
 		if (!ensureMsgf(Width <= MAX_VIEWPORT_SIZE && Height <= MAX_VIEWPORT_SIZE, TEXT("Invalid window with Width=%u and Height=%u"), Width, Height))
