Copying Private-LoadTimes-4.24 stream to Main. Biggest changes are in Materials/Shader memory freezing.

diff --git a/Engine/Source/Runtime/SlateRHIRenderer/Private/SlateRHIRenderer.cpp b/Engine/Source/Runtime/SlateRHIRenderer/Private/SlateRHIRenderer.cpp
--- a/Engine/Source/Runtime/SlateRHIRenderer/Private/SlateRHIRenderer.cpp
+++ b/Engine/Source/Runtime/SlateRHIRenderer/Private/SlateRHIRenderer.cpp
@@ -547,16 +547,8 @@ public:
 			OutputDeviceValue = FMath::Max(OutputDeviceValue, 2);
 		}
 
-		SetShaderValue(RHICmdList, GetPixelShader(), OutputDevice, OutputDeviceValue);
-		SetShaderValue(RHICmdList, GetPixelShader(), OutputGamut, OutputGamutValue);
-	}
-
-	virtual bool Serialize(FArchive& Ar) override
-	{
-		bool bShaderHasOutdatedParameters = FGlobalShader::Serialize(Ar);
-		Ar << OutputDevice;
-		Ar << OutputGamut;
-		return bShaderHasOutdatedParameters;
+		SetShaderValue(RHICmdList, RHICmdList.GetBoundPixelShader(), OutputDevice, OutputDeviceValue);
+		SetShaderValue(RHICmdList, RHICmdList.GetBoundPixelShader(), OutputGamut, OutputGamutValue);
 	}
 
 	static const TCHAR* GetSourceFilename()
@@ -570,8 +562,8 @@ public:
 	}
 
 private:
-	FShaderParameter OutputDevice;
-	FShaderParameter OutputGamut;
+	LAYOUT_FIELD(FShaderParameter, OutputDevice);
+	LAYOUT_FIELD(FShaderParameter, OutputGamut);
 };
 
 IMPLEMENT_SHADER_TYPE(, FCompositeLUTGenerationPS, TEXT("/Engine/Private/CompositeUIPixelShader.usf"), TEXT("GenerateLUTPS"), SF_Pixel);
@@ -607,15 +599,15 @@ public:
 	{
 		static const auto CVarOutputDevice = IConsoleManager::Get().FindTConsoleVariableDataInt(TEXT("r.HDR.Display.OutputDevice"));
 
-		SetTextureParameter(RHICmdList, GetPixelShader(), UITexture, UISampler, TStaticSamplerState<SF_Point>::GetRHI(), UITextureRHI);
-		SetTextureParameter(RHICmdList, GetPixelShader(), SceneTexture, SceneSampler, TStaticSamplerState<SF_Point>::GetRHI(), SceneTextureRHI);
-		SetTextureParameter(RHICmdList, GetPixelShader(), ColorSpaceLUT, ColorSpaceLUTSampler, TStaticSamplerState<SF_Bilinear, AM_Clamp, AM_Clamp, AM_Clamp>::GetRHI(), ColorSpaceLUTRHI);
-		SetShaderValue(RHICmdList, GetPixelShader(), UILevel, CVarUILevel.GetValueOnRenderThread());
-		SetShaderValue(RHICmdList, GetPixelShader(), OutputDevice, CVarOutputDevice->GetValueOnRenderThread());
+		SetTextureParameter(RHICmdList, RHICmdList.GetBoundPixelShader(), UITexture, UISampler, TStaticSamplerState<SF_Point>::GetRHI(), UITextureRHI);
+		SetTextureParameter(RHICmdList, RHICmdList.GetBoundPixelShader(), SceneTexture, SceneSampler, TStaticSamplerState<SF_Point>::GetRHI(), SceneTextureRHI);
+		SetTextureParameter(RHICmdList, RHICmdList.GetBoundPixelShader(), ColorSpaceLUT, ColorSpaceLUTSampler, TStaticSamplerState<SF_Bilinear, AM_Clamp, AM_Clamp, AM_Clamp>::GetRHI(), ColorSpaceLUTRHI);
+		SetShaderValue(RHICmdList, RHICmdList.GetBoundPixelShader(), UILevel, CVarUILevel.GetValueOnRenderThread());
+		SetShaderValue(RHICmdList, RHICmdList.GetBoundPixelShader(), OutputDevice, CVarOutputDevice->GetValueOnRenderThread());
 		
 		if (RHISupportsRenderTargetWriteMask(GMaxRHIShaderPlatform))
 		{
-			SetTextureParameter(RHICmdList, GetPixelShader(), UIWriteMaskTexture, UITextureWriteMaskRHI);
+			SetTextureParameter(RHICmdList, RHICmdList.GetBoundPixelShader(), UIWriteMaskTexture, UITextureWriteMaskRHI);
 		}
 	}
 
@@ -625,17 +617,6 @@ public:
 		OutEnvironment.SetDefine(TEXT("SCRGB_ENCODING"), EncodingType);
 	}
 
-	virtual bool Serialize(FArchive& Ar) override
-	{
-		bool bShaderHasOutdatedParameters = FGlobalShader::Serialize(Ar);
-		Ar << UITexture << UISampler;
-		Ar << UIWriteMaskTexture;
-		Ar << SceneTexture << SceneSampler;
-		Ar << ColorSpaceLUT << ColorSpaceLUTSampler;
-		Ar << UILevel << OutputDevice;
-		return bShaderHasOutdatedParameters;
-	}
-
 	static const TCHAR* GetSourceFilename()
 	{
 		return TEXT("/Engine/Private/CompositeUIPixelShader.usf");
@@ -647,15 +628,15 @@ public:
 	}
 
 private:
-	FShaderResourceParameter UITexture;
-	FShaderResourceParameter UIWriteMaskTexture;
-	FShaderResourceParameter UISampler;
-	FShaderResourceParameter SceneTexture;
-	FShaderResourceParameter SceneSampler;
-	FShaderResourceParameter ColorSpaceLUT;
-	FShaderResourceParameter ColorSpaceLUTSampler;
-	FShaderParameter UILevel;
-	FShaderParameter OutputDevice;
+	LAYOUT_FIELD(FShaderResourceParameter, UITexture);
+	LAYOUT_FIELD(FShaderResourceParameter, UIWriteMaskTexture);
+	LAYOUT_FIELD(FShaderResourceParameter, UISampler);
+	LAYOUT_FIELD(FShaderResourceParameter, SceneTexture);
+	LAYOUT_FIELD(FShaderResourceParameter, SceneSampler);
+	LAYOUT_FIELD(FShaderResourceParameter, ColorSpaceLUT);
+	LAYOUT_FIELD(FShaderResourceParameter, ColorSpaceLUTSampler);
+	LAYOUT_FIELD(FShaderParameter, UILevel);
+	LAYOUT_FIELD(FShaderParameter, OutputDevice);
 };
 
 #define SHADER_VARIATION(A) typedef FCompositePS<A> FCompositePS##A; \
@@ -920,11 +901,11 @@ void FSlateRHIRenderer::DrawWindow_RenderThread(FRHICommandListImmediate& RHICmd
 						const FVolumeBounds VolumeBounds(CompositionLUTSize);
 
 						GraphicsPSOInit.BoundShaderState.VertexDeclarationRHI = GScreenVertexDeclaration.VertexDeclarationRHI;
-						GraphicsPSOInit.BoundShaderState.VertexShaderRHI = GETSAFERHISHADER_VERTEX(*VertexShader);
+						GraphicsPSOInit.BoundShaderState.VertexShaderRHI = VertexShader.GetVertexShader();
 #if PLATFORM_SUPPORTS_GEOMETRY_SHADERS
-						GraphicsPSOInit.BoundShaderState.GeometryShaderRHI = GETSAFERHISHADER_GEOMETRY(*GeometryShader);
+						GraphicsPSOInit.BoundShaderState.GeometryShaderRHI = GeometryShader.GetGeometryShader();
 #endif
-						GraphicsPSOInit.BoundShaderState.PixelShaderRHI = GETSAFERHISHADER_PIXEL(*PixelShader);
+						GraphicsPSOInit.BoundShaderState.PixelShaderRHI = PixelShader.GetPixelShader();
 						GraphicsPSOInit.PrimitiveType = PT_TriangleStrip;
 						SetGraphicsPipelineState(RHICmdList, GraphicsPSOInit);
 
@@ -950,7 +931,7 @@ void FSlateRHIRenderer::DrawWindow_RenderThread(FRHICommandListImmediate& RHICmd
 
 					if (RHISupportsRenderTargetWriteMask(GMaxRHIShaderPlatform))
 					{
-						IPooledRenderTarget* RenderTargets[1] = { ViewportInfo.UITargetRT.GetReference() };
+						IPooledRenderTarget* RenderTargets[] = { ViewportInfo.UITargetRT.GetReference() };
 						FRenderTargetWriteMask::Decode<1>(RHICmdList, ShaderMap, RenderTargets, ViewportInfo.UITargetRTMask, 0, TEXT("UIRTWriteMask"));
 					}
 
@@ -974,8 +955,8 @@ void FSlateRHIRenderer::DrawWindow_RenderThread(FRHICommandListImmediate& RHICmd
 							TShaderMapRef<FCompositePS<1>> PixelShader(ShaderMap);
 
 							GraphicsPSOInit.BoundShaderState.VertexDeclarationRHI = GFilterVertexDeclaration.VertexDeclarationRHI;
-							GraphicsPSOInit.BoundShaderState.VertexShaderRHI = GETSAFERHISHADER_VERTEX(*VertexShader);
-							GraphicsPSOInit.BoundShaderState.PixelShaderRHI = GETSAFERHISHADER_PIXEL(*PixelShader);
+							GraphicsPSOInit.BoundShaderState.VertexShaderRHI = VertexShader.GetVertexShader();
+							GraphicsPSOInit.BoundShaderState.PixelShaderRHI = PixelShader.GetPixelShader();
 							GraphicsPSOInit.PrimitiveType = PT_TriangleList;
 
 							SetGraphicsPipelineState(RHICmdList, GraphicsPSOInit);
@@ -988,8 +969,8 @@ void FSlateRHIRenderer::DrawWindow_RenderThread(FRHICommandListImmediate& RHICmd
 							TShaderMapRef<FCompositePS<0>> PixelShader(ShaderMap);
 
 							GraphicsPSOInit.BoundShaderState.VertexDeclarationRHI = GFilterVertexDeclaration.VertexDeclarationRHI;
-							GraphicsPSOInit.BoundShaderState.VertexShaderRHI = GETSAFERHISHADER_VERTEX(*VertexShader);
-							GraphicsPSOInit.BoundShaderState.PixelShaderRHI = GETSAFERHISHADER_PIXEL(*PixelShader);
+							GraphicsPSOInit.BoundShaderState.VertexShaderRHI = VertexShader.GetVertexShader();
+							GraphicsPSOInit.BoundShaderState.PixelShaderRHI = PixelShader.GetPixelShader();
 							GraphicsPSOInit.PrimitiveType = PT_TriangleList;
 
 							SetGraphicsPipelineState(RHICmdList, GraphicsPSOInit);
@@ -1005,7 +986,7 @@ void FSlateRHIRenderer::DrawWindow_RenderThread(FRHICommandListImmediate& RHICmd
 							ViewportWidth, ViewportHeight,
 							FIntPoint(ViewportWidth, ViewportHeight),
 							FIntPoint(ViewportWidth, ViewportHeight),
-							*VertexShader,
+							VertexShader,
 							EDRF_UseTriangleOptimization);
 					}
 					RHICmdList.EndRenderPass();
@@ -1036,8 +1017,8 @@ void FSlateRHIRenderer::DrawWindow_RenderThread(FRHICommandListImmediate& RHICmd
 				TShaderMapRef<FScreenVS> VertexShader(ShaderMap);
 
 				GraphicsPSOInit.BoundShaderState.VertexDeclarationRHI = GFilterVertexDeclaration.VertexDeclarationRHI;
-				GraphicsPSOInit.BoundShaderState.VertexShaderRHI = GETSAFERHISHADER_VERTEX(*VertexShader);
-				GraphicsPSOInit.BoundShaderState.PixelShaderRHI = GETSAFERHISHADER_PIXEL(*PixelShader);
+				GraphicsPSOInit.BoundShaderState.VertexShaderRHI = VertexShader.GetVertexShader();
+				GraphicsPSOInit.BoundShaderState.PixelShaderRHI = PixelShader.GetPixelShader();
 				GraphicsPSOInit.PrimitiveType = PT_TriangleList;
 
 				SetGraphicsPipelineState(RHICmdList, GraphicsPSOInit);
@@ -1055,7 +1036,7 @@ void FSlateRHIRenderer::DrawWindow_RenderThread(FRHICommandListImmediate& RHICmd
 					ViewportWidth, ViewportHeight,
 					FIntPoint(ViewportWidth, ViewportHeight),
 					FIntPoint(ViewportWidth, ViewportHeight),
-					*VertexShader,
+					VertexShader,
 					EDRF_UseTriangleOptimization);
 				
 				RHICmdList.EndRenderPass();
diff --git a/Engine/Source/Runtime/SlateRHIRenderer/Private/SlateRHIRenderingPolicy.cpp b/Engine/Source/Runtime/SlateRHIRenderer/Private/SlateRHIRenderingPolicy.cpp
--- a/Engine/Source/Runtime/SlateRHIRenderer/Private/SlateRHIRenderingPolicy.cpp
+++ b/Engine/Source/Runtime/SlateRHIRenderer/Private/SlateRHIRenderingPolicy.cpp
@@ -405,7 +405,7 @@ static bool UpdateScissorRect(
 					RHICmdList.BeginRenderPass(RPInfo, TEXT("SlateUpdateScissorRect_ClearStencil"));
 				}
 
-				TShaderMap<FGlobalShaderType>* MaxFeatureLevelShaderMap = GetGlobalShaderMap(GMaxRHIFeatureLevel);
+				FGlobalShaderMap* MaxFeatureLevelShaderMap = GetGlobalShaderMap(GMaxRHIFeatureLevel);
 
 				// Set the new shaders
 				TShaderMapRef<FSlateMaskingVS> VertexShader(MaxFeatureLevelShaderMap);
@@ -435,8 +435,8 @@ static bool UpdateScissorRect(
 						, /*StencilWriteMask*/ 0xFF>::GetRHI();
 
 					WriteMaskPSOInit.BoundShaderState.VertexDeclarationRHI = GSlateMaskingVertexDeclaration.VertexDeclarationRHI;
-					WriteMaskPSOInit.BoundShaderState.VertexShaderRHI = GETSAFERHISHADER_VERTEX(*VertexShader);
-					WriteMaskPSOInit.BoundShaderState.PixelShaderRHI = GETSAFERHISHADER_PIXEL(*PixelShader);
+					WriteMaskPSOInit.BoundShaderState.VertexShaderRHI = VertexShader.GetVertexShader();
+					WriteMaskPSOInit.BoundShaderState.PixelShaderRHI = PixelShader.GetPixelShader();
 					WriteMaskPSOInit.PrimitiveType = PT_TriangleStrip;
 
 					SetGraphicsPipelineState(RHICmdList, WriteMaskPSOInit);
@@ -651,7 +651,7 @@ void FSlateRHIRenderingPolicy::DrawElements(
 	const FSlateRenderDataHandle* LastHandle = nullptr;
 
 	const ERHIFeatureLevel::Type FeatureLevel = GMaxRHIFeatureLevel;
-	TShaderMap<FGlobalShaderType>* ShaderMap = GetGlobalShaderMap(FeatureLevel);
+	FGlobalShaderMap* ShaderMap = GetGlobalShaderMap(FeatureLevel);
 
 #if WITH_SLATE_VISUALIZERS
 	FRandomStream BatchColors(1337);
@@ -774,16 +774,16 @@ void FSlateRHIRenderingPolicy::DrawElements(
 			{
 				check(RHICmdList.IsInsideRenderPass());
 				check(RenderBatch.NumIndices > 0);
-				FSlateElementPS* PixelShader = nullptr;
+				TShaderRef<FSlateElementPS> PixelShader;
 
 				const bool bUseInstancing = RenderBatch.InstanceCount > 1 && RenderBatch.InstanceData != nullptr;
 				check(bUseInstancing == false);
 
 #if WITH_SLATE_VISUALIZERS
-				FSlateDebugBatchingPS* BatchingPixelShader = nullptr;
+				TShaderRef<FSlateDebugBatchingPS> BatchingPixelShader;
 				if (CVarShowSlateBatching.GetValueOnRenderThread() != 0)
 				{
-					BatchingPixelShader = *TShaderMapRef<FSlateDebugBatchingPS>(ShaderMap);
+					BatchingPixelShader = TShaderMapRef<FSlateDebugBatchingPS>(ShaderMap);
 					PixelShader = BatchingPixelShader;
 				}
 				else
@@ -828,8 +828,8 @@ void FSlateRHIRenderingPolicy::DrawElements(
 				}
 
 				GraphicsPSOInit.BoundShaderState.VertexDeclarationRHI = GSlateVertexDeclaration.VertexDeclarationRHI;
-				GraphicsPSOInit.BoundShaderState.VertexShaderRHI = GETSAFERHISHADER_VERTEX(*GlobalVertexShader);
-				GraphicsPSOInit.BoundShaderState.PixelShaderRHI = GETSAFERHISHADER_PIXEL(PixelShader);
+				GraphicsPSOInit.BoundShaderState.VertexShaderRHI = GlobalVertexShader.GetVertexShader();
+				GraphicsPSOInit.BoundShaderState.PixelShaderRHI = PixelShader.GetPixelShader();
 				GraphicsPSOInit.PrimitiveType = GetRHIPrimitiveType(RenderBatch.DrawPrimitiveType);
 
 				SetGraphicsPipelineState(RHICmdList, GraphicsPSOInit);
@@ -1020,32 +1020,32 @@ void FSlateRHIRenderingPolicy::DrawElements(
 
 					const FMaterial* Material = MaterialRenderProxy->GetMaterial(ActiveSceneView.GetFeatureLevel());
 
-					FSlateMaterialShaderPS* PixelShader = GetMaterialPixelShader(Material, ShaderType);
+					TShaderRef<FSlateMaterialShaderPS> PixelShader = GetMaterialPixelShader(Material, ShaderType);
 
 					const bool bUseInstancing = RenderBatch.InstanceCount > 0 && RenderBatch.InstanceData != nullptr;
-					FSlateMaterialShaderVS* VertexShader = GetMaterialVertexShader(Material, bUseInstancing);
+					TShaderRef<FSlateMaterialShaderVS> VertexShader = GetMaterialVertexShader(Material, bUseInstancing);
 
-					if (VertexShader && PixelShader)
+					if (VertexShader.IsValid() && PixelShader.IsValid())
 					{
 #if WITH_SLATE_VISUALIZERS
 						if (CVarShowSlateBatching.GetValueOnRenderThread() != 0)
 						{
-							FSlateDebugBatchingPS* BatchingPixelShader = *TShaderMapRef<FSlateDebugBatchingPS>(ShaderMap);
+							TShaderMapRef<FSlateDebugBatchingPS> BatchingPixelShader(ShaderMap);
 
 							GraphicsPSOInit.BoundShaderState.VertexDeclarationRHI = bUseInstancing ? GSlateInstancedVertexDeclaration.VertexDeclarationRHI : GSlateVertexDeclaration.VertexDeclarationRHI;
-							GraphicsPSOInit.BoundShaderState.VertexShaderRHI = GETSAFERHISHADER_VERTEX(*GlobalVertexShader);
-							GraphicsPSOInit.BoundShaderState.PixelShaderRHI = GETSAFERHISHADER_PIXEL(BatchingPixelShader);
+							GraphicsPSOInit.BoundShaderState.VertexShaderRHI = GlobalVertexShader.GetVertexShader();
+							GraphicsPSOInit.BoundShaderState.PixelShaderRHI = BatchingPixelShader.GetPixelShader();
 
 							BatchingPixelShader->SetBatchColor(RHICmdList, BatchColor);
 							GraphicsPSOInit.BlendState = TStaticBlendState<CW_RGB, BO_Add, BF_One, BF_One, BO_Add, BF_Zero, BF_InverseSourceAlpha>::GetRHI();
 						}
 						else if (CVarShowSlateOverdraw.GetValueOnRenderThread() != 0)
 						{
-							FSlateElementPS* OverdrawPixelShader = *TShaderMapRef<FSlateDebugOverdrawPS>(ShaderMap);
+							TShaderMapRef<FSlateDebugOverdrawPS> OverdrawPixelShader(ShaderMap);
 
 							GraphicsPSOInit.BoundShaderState.VertexDeclarationRHI = bUseInstancing ? GSlateInstancedVertexDeclaration.VertexDeclarationRHI : GSlateVertexDeclaration.VertexDeclarationRHI;
-							GraphicsPSOInit.BoundShaderState.VertexShaderRHI = GETSAFERHISHADER_VERTEX(*GlobalVertexShader);
-							GraphicsPSOInit.BoundShaderState.PixelShaderRHI = GETSAFERHISHADER_PIXEL(OverdrawPixelShader);
+							GraphicsPSOInit.BoundShaderState.VertexShaderRHI = GlobalVertexShader.GetVertexShader();
+							GraphicsPSOInit.BoundShaderState.PixelShaderRHI = OverdrawPixelShader.GetPixelShader();
 
 							GraphicsPSOInit.BlendState = TStaticBlendState<CW_RGB, BO_Add, BF_One, BF_One, BO_Add, BF_Zero, BF_InverseSourceAlpha>::GetRHI();
 						}
@@ -1060,8 +1060,8 @@ void FSlateRHIRenderingPolicy::DrawElements(
 							}
 
 							GraphicsPSOInit.BoundShaderState.VertexDeclarationRHI = bUseInstancing ? GSlateInstancedVertexDeclaration.VertexDeclarationRHI : GSlateVertexDeclaration.VertexDeclarationRHI;
-							GraphicsPSOInit.BoundShaderState.VertexShaderRHI = GETSAFERHISHADER_VERTEX(VertexShader);
-							GraphicsPSOInit.BoundShaderState.PixelShaderRHI = GETSAFERHISHADER_PIXEL(PixelShader);
+							GraphicsPSOInit.BoundShaderState.VertexShaderRHI = VertexShader.GetVertexShader();
+							GraphicsPSOInit.BoundShaderState.PixelShaderRHI = PixelShader.GetPixelShader();
 							GraphicsPSOInit.PrimitiveType = GetRHIPrimitiveType(RenderBatch.DrawPrimitiveType);
 
 							SetGraphicsPipelineState(RHICmdList, GraphicsPSOInit);
@@ -1268,15 +1268,15 @@ ETextureSamplerFilter FSlateRHIRenderingPolicy::GetSamplerFilter(const UTexture*
 	return Filter;
 }
 
-FSlateElementPS* FSlateRHIRenderingPolicy::GetTexturePixelShader( TShaderMap<FGlobalShaderType>* ShaderMap, ESlateShader ShaderType, ESlateDrawEffect DrawEffects )
+TShaderRef<FSlateElementPS> FSlateRHIRenderingPolicy::GetTexturePixelShader( FGlobalShaderMap* ShaderMap, ESlateShader ShaderType, ESlateDrawEffect DrawEffects )
 {
-	FSlateElementPS* PixelShader = nullptr;
+	TShaderRef<FSlateElementPS> PixelShader;
 	QUICK_SCOPE_CYCLE_COUNTER(STAT_Slate_GetTexturePixelShader);
 
 #if WITH_SLATE_VISUALIZERS
 	if ( CVarShowSlateOverdraw.GetValueOnRenderThread() != 0 )
 	{
-		PixelShader = *TShaderMapRef<FSlateDebugOverdrawPS>(ShaderMap);
+		PixelShader = TShaderMapRef<FSlateDebugOverdrawPS>(ShaderMap);
 	}
 	else
 #endif
@@ -1292,31 +1292,31 @@ FSlateElementPS* FSlateRHIRenderingPolicy::GetTexturePixelShader( TShaderMap<FGl
 		case ESlateShader::Default:
 			if ( bUseTextureAlpha )
 			{
-				PixelShader = *TShaderMapRef<TSlateElementPS<ESlateShader::Default, true, true> >(ShaderMap);
+				PixelShader = TShaderMapRef<TSlateElementPS<ESlateShader::Default, true, true> >(ShaderMap);
 			}
 			else
 			{
-				PixelShader = *TShaderMapRef<TSlateElementPS<ESlateShader::Default, true, false> >(ShaderMap);
+				PixelShader = TShaderMapRef<TSlateElementPS<ESlateShader::Default, true, false> >(ShaderMap);
 			}
 			break;
 		case ESlateShader::Border:
 			if ( bUseTextureAlpha )
 			{
-				PixelShader = *TShaderMapRef<TSlateElementPS<ESlateShader::Border, true, true> >(ShaderMap);
+				PixelShader = TShaderMapRef<TSlateElementPS<ESlateShader::Border, true, true> >(ShaderMap);
 			}
 			else
 			{
-				PixelShader = *TShaderMapRef<TSlateElementPS<ESlateShader::Border, true, false> >(ShaderMap);
+				PixelShader = TShaderMapRef<TSlateElementPS<ESlateShader::Border, true, false> >(ShaderMap);
 			}
 			break;
 		case ESlateShader::GrayscaleFont:
-			PixelShader = *TShaderMapRef<TSlateElementPS<ESlateShader::GrayscaleFont, true> >(ShaderMap);
+			PixelShader = TShaderMapRef<TSlateElementPS<ESlateShader::GrayscaleFont, true> >(ShaderMap);
 			break;
 		case ESlateShader::ColorFont:
-			PixelShader = *TShaderMapRef<TSlateElementPS<ESlateShader::ColorFont, true> >(ShaderMap);
+			PixelShader = TShaderMapRef<TSlateElementPS<ESlateShader::ColorFont, true> >(ShaderMap);
 			break;
 		case ESlateShader::LineSegment:
-			PixelShader = *TShaderMapRef<TSlateElementPS<ESlateShader::LineSegment, true> >(ShaderMap);
+			PixelShader = TShaderMapRef<TSlateElementPS<ESlateShader::LineSegment, true> >(ShaderMap);
 			break;
 		}
 	}
@@ -1328,31 +1328,31 @@ FSlateElementPS* FSlateRHIRenderingPolicy::GetTexturePixelShader( TShaderMap<FGl
 		case ESlateShader::Default:
 			if ( bUseTextureAlpha )
 			{
-				PixelShader = *TShaderMapRef<TSlateElementPS<ESlateShader::Default, false, true> >(ShaderMap);
+				PixelShader = TShaderMapRef<TSlateElementPS<ESlateShader::Default, false, true> >(ShaderMap);
 			}
 			else
 			{
-				PixelShader = *TShaderMapRef<TSlateElementPS<ESlateShader::Default, false, false> >(ShaderMap);
+				PixelShader = TShaderMapRef<TSlateElementPS<ESlateShader::Default, false, false> >(ShaderMap);
 			}
 			break;
 		case ESlateShader::Border:
 			if ( bUseTextureAlpha )
 			{
-				PixelShader = *TShaderMapRef<TSlateElementPS<ESlateShader::Border, false, true> >(ShaderMap);
+				PixelShader = TShaderMapRef<TSlateElementPS<ESlateShader::Border, false, true> >(ShaderMap);
 			}
 			else
 			{
-				PixelShader = *TShaderMapRef<TSlateElementPS<ESlateShader::Border, false, false> >(ShaderMap);
+				PixelShader = TShaderMapRef<TSlateElementPS<ESlateShader::Border, false, false> >(ShaderMap);
 			}
 			break;
 		case ESlateShader::GrayscaleFont:
-			PixelShader = *TShaderMapRef<TSlateElementPS<ESlateShader::GrayscaleFont, false> >(ShaderMap);
+			PixelShader = TShaderMapRef<TSlateElementPS<ESlateShader::GrayscaleFont, false> >(ShaderMap);
 			break;
 		case ESlateShader::ColorFont:
-			PixelShader = *TShaderMapRef<TSlateElementPS<ESlateShader::ColorFont, false> >(ShaderMap);
+			PixelShader = TShaderMapRef<TSlateElementPS<ESlateShader::ColorFont, false> >(ShaderMap);
 			break;
 		case ESlateShader::LineSegment:
-			PixelShader = *TShaderMapRef<TSlateElementPS<ESlateShader::LineSegment, false> >(ShaderMap);
+			PixelShader = TShaderMapRef<TSlateElementPS<ESlateShader::LineSegment, false> >(ShaderMap);
 			break;
 		}
 	}
@@ -1363,11 +1363,11 @@ FSlateElementPS* FSlateRHIRenderingPolicy::GetTexturePixelShader( TShaderMap<FGl
 	return PixelShader;
 }
 
-FSlateMaterialShaderPS* FSlateRHIRenderingPolicy::GetMaterialPixelShader( const FMaterial* Material, ESlateShader ShaderType)
+TShaderRef<FSlateMaterialShaderPS> FSlateRHIRenderingPolicy::GetMaterialPixelShader( const FMaterial* Material, ESlateShader ShaderType)
 {
 	const FMaterialShaderMap* MaterialShaderMap = Material->GetRenderingThreadShaderMap();
 
-	FShader* FoundShader = nullptr;
+	TShaderRef<FShader> FoundShader;
 	switch (ShaderType)
 	{
 	case ESlateShader::Default:
@@ -1390,14 +1390,14 @@ FSlateMaterialShaderPS* FSlateRHIRenderingPolicy::GetMaterialPixelShader( const
 		break;
 	}
 
-	return FoundShader ? (FSlateMaterialShaderPS*)FoundShader->GetShader() : nullptr;
+	return TShaderRef<FSlateMaterialShaderPS>::Cast(FoundShader);
 }
 
-FSlateMaterialShaderVS* FSlateRHIRenderingPolicy::GetMaterialVertexShader( const FMaterial* Material, bool bUseInstancing )
+TShaderRef<FSlateMaterialShaderVS> FSlateRHIRenderingPolicy::GetMaterialVertexShader( const FMaterial* Material, bool bUseInstancing )
 {
 	const FMaterialShaderMap* MaterialShaderMap = Material->GetRenderingThreadShaderMap();
 
-	FShader* FoundShader = nullptr;
+	TShaderRef<FShader> FoundShader;
 	if( bUseInstancing )
 	{
 		FoundShader = MaterialShaderMap->GetShader(&TSlateMaterialShaderVS<true>::StaticType);
@@ -1407,7 +1407,7 @@ FSlateMaterialShaderVS* FSlateRHIRenderingPolicy::GetMaterialVertexShader( const
 		FoundShader = MaterialShaderMap->GetShader(&TSlateMaterialShaderVS<false>::StaticType);
 	}
 	
-	return FoundShader ? (FSlateMaterialShaderVS*)FoundShader->GetShader() : nullptr;
+	return TShaderRef<FSlateMaterialShaderVS>::Cast(FoundShader);
 }
 
 EPrimitiveType FSlateRHIRenderingPolicy::GetRHIPrimitiveType(ESlateDrawPrimitive SlateType)
diff --git a/Engine/Source/Runtime/SlateRHIRenderer/Private/SlateRHIRenderingPolicy.h b/Engine/Source/Runtime/SlateRHIRenderer/Private/SlateRHIRenderingPolicy.h
--- a/Engine/Source/Runtime/SlateRHIRenderer/Private/SlateRHIRenderingPolicy.h
+++ b/Engine/Source/Runtime/SlateRHIRenderer/Private/SlateRHIRenderingPolicy.h
@@ -13,12 +13,16 @@
 #include "SlateElementVertexBuffer.h"
 #include "SlateRHIResourceManager.h"
 #include "Shader.h"
+#include "GlobalShader.h"
 #include "Engine/TextureLODSettings.h"
 
 class FSlateFontServices;
 class FSlateRHIResourceManager;
 class FSlatePostProcessor;
 class UDeviceProfile;
+class FSlateElementPS;
+class FSlateMaterialShaderPS;
+class FSlateMaterialShaderVS;
 
 struct FSlateRenderingParams
 {
@@ -80,9 +84,9 @@ private:
 	 * @param DrawEffects	Draw effects being used
 	 * @return The pixel shader for use with the shader type and draw effects
 	 */
-	class FSlateElementPS* GetTexturePixelShader( TShaderMap<FGlobalShaderType>* ShaderMap, ESlateShader ShaderType, ESlateDrawEffect DrawEffects );
-	class FSlateMaterialShaderPS* GetMaterialPixelShader( const class FMaterial* Material, ESlateShader ShaderType);
-	class FSlateMaterialShaderVS* GetMaterialVertexShader( const class FMaterial* Material, bool bUseInstancing );
+	TShaderRef<FSlateElementPS> GetTexturePixelShader(FGlobalShaderMap* ShaderMap, ESlateShader ShaderType, ESlateDrawEffect DrawEffects );
+	TShaderRef<FSlateMaterialShaderPS> GetMaterialPixelShader( const class FMaterial* Material, ESlateShader ShaderType );
+	TShaderRef<FSlateMaterialShaderVS> GetMaterialVertexShader( const class FMaterial* Material, bool bUseInstancing );
 
 	/** @return The RHI primitive type from the Slate primitive type */
 	EPrimitiveType GetRHIPrimitiveType(ESlateDrawPrimitive SlateType);
diff --git a/Engine/Source/Runtime/SlateRHIRenderer/Private/SlateRHIResourceManager.cpp b/Engine/Source/Runtime/SlateRHIRenderer/Private/SlateRHIResourceManager.cpp
--- a/Engine/Source/Runtime/SlateRHIRenderer/Private/SlateRHIResourceManager.cpp
+++ b/Engine/Source/Runtime/SlateRHIRenderer/Private/SlateRHIResourceManager.cpp
@@ -24,6 +24,7 @@
 #include "Slate/SlateTextureAtlasInterface.h"
 #include "SlateAtlasedTextureResource.h"
 #include "ImageUtils.h"
+#include "Async/ParallelFor.h"
 
 #define LOCTEXT_NAMESPACE "Slate"
 
@@ -310,45 +311,58 @@ void FSlateRHIResourceManager::Tick(float DeltaSeconds)
 
 void FSlateRHIResourceManager::CreateTextures( const TArray< const FSlateBrush* >& Resources )
 {
+	FModuleManager::LoadModuleChecked<IImageWrapperModule>(FName("ImageWrapper"));
+
 	TMap<FName,FNewTextureInfo> TextureInfoMap;
+	FCriticalSection Lock;
+
+	// reserve space so we don't reallocate memory during the ParallelFor
+	TextureInfoMap.Reserve(Resources.Num());
 
-	const uint32 Stride = GPixelFormats[PF_R8G8B8A8].BlockBytes;
-	for( int32 ResourceIndex = 0; ResourceIndex < Resources.Num(); ++ResourceIndex )
+	ParallelFor(Resources.Num(), [this, &Resources, &Lock, &TextureInfoMap](int32 ResourceIndex)
 	{
+		const uint32 Stride = GPixelFormats[PF_R8G8B8A8].BlockBytes;
 		const FSlateBrush& Brush = *Resources[ResourceIndex];
 		const FName TextureName = Brush.GetResourceName();
 		if( TextureName != NAME_None && !Brush.HasUObject() && !Brush.IsDynamicallyLoaded() && !ResourceMap.Contains(TextureName) )
 		{
 			// Find the texture or add it if it doesnt exist (only load the texture once)
-			FNewTextureInfo& Info = TextureInfoMap.FindOrAdd( TextureName );
+			FNewTextureInfo* Info;
+			{
+				FScopeLock Locker(&Lock);
+				Info = &TextureInfoMap.FindOrAdd(TextureName);
+			}
 	
-			Info.bSrgb = (Brush.ImageType != ESlateBrushImageType::Linear);
+			Info->bSrgb = (Brush.ImageType != ESlateBrushImageType::Linear);
 
 			// Only atlas the texture if none of the brushes that use it tile it and the image is srgb
 		
-			Info.bShouldAtlas &= ( Brush.Tiling == ESlateBrushTileType::NoTile && Info.bSrgb && AtlasSize > 0 );
+			Info->bShouldAtlas &= ( Brush.Tiling == ESlateBrushTileType::NoTile && Info->bSrgb && AtlasSize > 0 );
 
 			// Texture has been loaded if the texture data is valid
-			if( !Info.TextureData.IsValid() )
+			if( !Info->TextureData.IsValid() )
 			{
 				uint32 Width = 0;
 				uint32 Height = 0;
 				TArray<uint8> RawData;
-				bool bSucceeded = LoadTexture( Brush, Width, Height, RawData );
-
-				Info.TextureData = MakeShareable( new FSlateTextureData( Width, Height, Stride, RawData ) );
+				const bool bSucceeded = LoadTexture( Brush, Width, Height, RawData );
+				if (bSucceeded)
+				{
+					Info->TextureData = MakeShareable(new FSlateTextureData(Width, Height, Stride, RawData));
 
-				const bool bTooLargeForAtlas = (Width >= (uint32)MaxAltasedTextureSize.X || Height >= (uint32)MaxAltasedTextureSize.Y || Width >= AtlasSize || Height >= AtlasSize );
+					const bool bTooLargeForAtlas = (Width >= (uint32)MaxAltasedTextureSize.X || Height >= (uint32)MaxAltasedTextureSize.Y || Width >= AtlasSize || Height >= AtlasSize);
 
-				Info.bShouldAtlas &= !bTooLargeForAtlas;
+					Info->bShouldAtlas &= !bTooLargeForAtlas;
+				}
 
-				if( !bSucceeded || !ensureMsgf( Info.TextureData->GetRawBytes().Num() > 0, TEXT("Slate resource: (%s) contains no data"), *TextureName.ToString() ) )
+				if( !bSucceeded || !ensureMsgf( Info->TextureData->GetRawBytes().Num() > 0, TEXT("Slate resource: (%s) contains no data"), *TextureName.ToString() ) )
 				{
+					FScopeLock Locker(&Lock);
 					TextureInfoMap.Remove( TextureName );
 				}
 			}
 		}
-	}
+	});
 
 	// Sort textures by size.  The largest textures are atlased first which creates a more compact atlas
 	TextureInfoMap.ValueSort( FCompareFNewTextureInfoByTextureSize() );
@@ -380,7 +394,8 @@ bool FSlateRHIResourceManager::LoadTexture( const FSlateBrush& InBrush, uint32&
  */
 bool FSlateRHIResourceManager::LoadTexture( const FName& TextureName, const FString& ResourcePath, uint32& Width, uint32& Height, TArray<uint8>& DecodedImage )
 {
-	checkSlow( IsThreadSafeForSlateRendering() );
+	// @todo loadtime: nothing in this function uses SlateRendering, so why is this here? I think it's a bad check, and breaks ParallelFor above
+//	checkSlow( IsThreadSafeForSlateRendering() );
 
 	bool bSucceeded = false;
 	uint32 BytesPerPixel = 4;
diff --git a/Engine/Source/Runtime/SlateRHIRenderer/Private/SlateShaders.cpp b/Engine/Source/Runtime/SlateRHIRenderer/Private/SlateShaders.cpp
--- a/Engine/Source/Runtime/SlateRHIRenderer/Private/SlateShaders.cpp
+++ b/Engine/Source/Runtime/SlateRHIRenderer/Private/SlateShaders.cpp
@@ -10,6 +10,8 @@ int32 GSlateColorDeficiencySeverity = 0;
 bool GSlateColorDeficiencyCorrection = false;
 bool GSlateShowColorDeficiencyCorrectionWithDeficiency = false;
 
+IMPLEMENT_TYPE_LAYOUT(FSlateElementPS);
+
 IMPLEMENT_SHADER_TYPE(, FSlateElementVS, TEXT("/Engine/Private/SlateVertexShader.usf"), TEXT("Main"), SF_Vertex);
 
 IMPLEMENT_SHADER_TYPE(, FSlateDebugOverdrawPS, TEXT("/Engine/Private/SlateElementPixelShader.usf"), TEXT("DebugOverdrawMain"), SF_Pixel );
@@ -134,21 +136,21 @@ FSlateElementVS::FSlateElementVS( const ShaderMetaType::CompiledShaderInitialize
 
 void FSlateElementVS::SetViewProjection(FRHICommandList& RHICmdList, const FMatrix& InViewProjection )
 {
-	SetShaderValue(RHICmdList, GetVertexShader(), ViewProjection, InViewProjection );
+	SetShaderValue(RHICmdList, RHICmdList.GetBoundVertexShader(), ViewProjection, InViewProjection );
 }
 
 void FSlateElementVS::SetShaderParameters(FRHICommandList& RHICmdList, const FVector4& ShaderParams )
 {
-	SetShaderValue(RHICmdList, GetVertexShader(), VertexShaderParams, ShaderParams );
+	SetShaderValue(RHICmdList, RHICmdList.GetBoundVertexShader(), VertexShaderParams, ShaderParams );
 }
 
 void FSlateElementVS::SetVerticalAxisMultiplier(FRHICommandList& RHICmdList, float InMultiplier )
 {
-	SetShaderValue(RHICmdList, GetVertexShader(), SwitchVerticalAxisMultiplier, InMultiplier );
+	SetShaderValue(RHICmdList, RHICmdList.GetBoundVertexShader(), SwitchVerticalAxisMultiplier, InMultiplier );
 }
 
 /** Serializes the shader data */
-bool FSlateElementVS::Serialize( FArchive& Ar )
+/*bool FSlateElementVS::Serialize( FArchive& Ar )
 {
 	bool bShaderHasOutdatedParameters = FGlobalShader::Serialize( Ar );
 
@@ -157,7 +159,7 @@ bool FSlateElementVS::Serialize( FArchive& Ar )
 	Ar << SwitchVerticalAxisMultiplier;
 
 	return bShaderHasOutdatedParameters;
-}
+}*/
 
 
 /************************************************************************/
@@ -174,12 +176,12 @@ FSlateMaskingVS::FSlateMaskingVS(const ShaderMetaType::CompiledShaderInitializer
 
 void FSlateMaskingVS::SetViewProjection(FRHICommandList& RHICmdList, const FMatrix& InViewProjection)
 {
-	SetShaderValue(RHICmdList, GetVertexShader(), ViewProjection, InViewProjection);
+	SetShaderValue(RHICmdList, RHICmdList.GetBoundVertexShader(), ViewProjection, InViewProjection);
 }
 
 void FSlateMaskingVS::SetVerticalAxisMultiplier(FRHICommandList& RHICmdList, float InMultiplier )
 {
-	SetShaderValue(RHICmdList, GetVertexShader(), SwitchVerticalAxisMultiplier, InMultiplier );
+	SetShaderValue(RHICmdList, RHICmdList.GetBoundVertexShader(), SwitchVerticalAxisMultiplier, InMultiplier );
 }
 
 void FSlateMaskingVS::SetMaskRect(FRHICommandList& RHICmdList, const FVector2D& TopLeft, const FVector2D& TopRight, const FVector2D& BotLeft, const FVector2D& BotRight)
@@ -187,11 +189,11 @@ void FSlateMaskingVS::SetMaskRect(FRHICommandList& RHICmdList, const FVector2D&
 	//FVector4 MaskRectVal[4] = { FVector4(TopLeft, FVector2D::ZeroVector), FVector4(TopRight, FVector2D::ZeroVector), FVector4(BotLeft, FVector2D::ZeroVector), FVector4(BotRight, FVector2D::ZeroVector) };
 	FVector4 MaskRectVal[2] = { FVector4(TopLeft, TopRight), FVector4(BotLeft, BotRight) };
 
-	SetShaderValue(RHICmdList, GetVertexShader(), MaskRect, MaskRectVal);
+	SetShaderValue(RHICmdList, RHICmdList.GetBoundVertexShader(), MaskRect, MaskRectVal);
 }
 
 /** Serializes the shader data */
-bool FSlateMaskingVS::Serialize(FArchive& Ar)
+/*bool FSlateMaskingVS::Serialize(FArchive& Ar)
 {
 	bool bShaderHasOutdatedParameters = FGlobalShader::Serialize(Ar);
 
@@ -200,4 +202,4 @@ bool FSlateMaskingVS::Serialize(FArchive& Ar)
 	Ar << SwitchVerticalAxisMultiplier;
 
 	return bShaderHasOutdatedParameters;
-}
+}*/
diff --git a/Engine/Source/Runtime/SlateRHIRenderer/Private/SlateShaders.h b/Engine/Source/Runtime/SlateRHIRenderer/Private/SlateShaders.h
--- a/Engine/Source/Runtime/SlateRHIRenderer/Private/SlateShaders.h
+++ b/Engine/Source/Runtime/SlateRHIRenderer/Private/SlateShaders.h
@@ -93,22 +93,23 @@ public:
 	void SetVerticalAxisMultiplier(FRHICommandList& RHICmdList, float InMultiplier);
 
 	/** Serializes the shader data */
-	virtual bool Serialize( FArchive& Ar ) override;
+	//virtual bool Serialize( FArchive& Ar ) override;
 
 private:
 	/** ViewProjection parameter used by the shader */
-	FShaderParameter ViewProjection;
+	LAYOUT_FIELD(FShaderParameter, ViewProjection)
 	/** Shader parmeters used by the shader */
-	FShaderParameter VertexShaderParams;
+	LAYOUT_FIELD(FShaderParameter, VertexShaderParams)
 	/** Parameter used to determine if we need to swtich the vertical axis for opengl */
-	FShaderParameter SwitchVerticalAxisMultiplier;
+	LAYOUT_FIELD(FShaderParameter, SwitchVerticalAxisMultiplier)
 };
 
 /** 
  * Base class slate pixel shader for all elements
  */
 class FSlateElementPS : public FGlobalShader
-{	
+{
+	DECLARE_TYPE_LAYOUT(FSlateElementPS, NonVirtual);
 public:
 	/** Indicates that this shader should be cached */
 	static bool ShouldCompilePermutation(const FGlobalShaderPermutationParameters& Parameters) 
@@ -140,7 +141,7 @@ public:
 	 */
 	void SetTexture(FRHICommandList& RHICmdList, FRHITexture* InTexture, const FSamplerStateRHIRef SamplerState )
 	{
-		SetTextureParameter(RHICmdList, GetPixelShader(), TextureParameter, TextureParameterSampler, SamplerState, InTexture );
+		SetTextureParameter(RHICmdList, RHICmdList.GetBoundPixelShader(), TextureParameter, TextureParameterSampler, SamplerState, InTexture );
 	}
 
 	/**
@@ -150,7 +151,7 @@ public:
 	 */
 	void SetShaderParams(FRHICommandList& RHICmdList, const FVector4& InShaderParams )
 	{
-		SetShaderValue(RHICmdList, GetPixelShader(), ShaderParams, InShaderParams );
+		SetShaderValue(RHICmdList, RHICmdList.GetBoundPixelShader(), ShaderParams, InShaderParams );
 	}
 
 	/**
@@ -162,26 +163,16 @@ public:
 	{
 		FVector4 Values( 2.2f / InDisplayGamma, 1.0f/InDisplayGamma, bInvertAlpha, InContrast);
 
-		SetShaderValue(RHICmdList, GetPixelShader(), GammaAndAlphaValues, Values);
+		SetShaderValue(RHICmdList, RHICmdList.GetBoundPixelShader(), GammaAndAlphaValues, Values);
 	}
 
-	virtual bool Serialize( FArchive& Ar )
-	{
-		bool bShaderHasOutdatedParameters = FGlobalShader::Serialize( Ar );
-
-		Ar << TextureParameter;
-		Ar << TextureParameterSampler;
-		Ar << ShaderParams;
-		Ar << GammaAndAlphaValues;
-
-		return bShaderHasOutdatedParameters;
-	}
 private:
+	
 	/** Texture parameter used by the shader */
-	FShaderResourceParameter TextureParameter;
-	FShaderResourceParameter TextureParameterSampler;
-	FShaderParameter ShaderParams;
-	FShaderParameter GammaAndAlphaValues;
+	LAYOUT_FIELD(FShaderResourceParameter, TextureParameter);
+	LAYOUT_FIELD(FShaderResourceParameter, TextureParameterSampler);
+	LAYOUT_FIELD(FShaderParameter, ShaderParams);
+	LAYOUT_FIELD(FShaderParameter, GammaAndAlphaValues);
 };
 
 /** 
@@ -217,14 +208,6 @@ public:
 
 		FSlateElementPS::ModifyCompilationEnvironment( Parameters, OutEnvironment );
 	}
-
-	/** Serializes the shader data */
-	virtual bool Serialize( FArchive& Ar ) override
-	{
-		return FSlateElementPS::Serialize( Ar );
-	}
-private:
-
 };
 
 /** 
@@ -249,13 +232,6 @@ public:
 		: FSlateElementPS( Initializer )
 	{
 	}
-
-
-	virtual bool Serialize( FArchive& Ar ) override
-	{
-		return FSlateElementPS::Serialize( Ar );
-	}
-private:
 };
 
 /** 
@@ -289,20 +265,11 @@ public:
 	*/
 	void SetBatchColor(FRHICommandList& RHICmdList, const FLinearColor& InBatchColor)
 	{
-		SetShaderValue(RHICmdList, GetPixelShader(), BatchColor, InBatchColor);
-	}
-
-	virtual bool Serialize( FArchive& Ar ) override
-	{
-		bool bShaderHasOutdatedParameters = FSlateElementPS::Serialize( Ar );
-
-		Ar << BatchColor;
-
-		return bShaderHasOutdatedParameters;
+		SetShaderValue(RHICmdList, RHICmdList.GetBoundPixelShader(), BatchColor, InBatchColor);
 	}
 
 private:
-	FShaderParameter BatchColor;
+	LAYOUT_FIELD(FShaderParameter, BatchColor);
 };
 
 const int32 MAX_BLUR_SAMPLES = 127;
@@ -333,35 +300,26 @@ public:
 
 	void SetBufferSizeAndDirection(FRHICommandList& RHICmdList, const FVector2D& InBufferSize, const FVector2D& InDir)
 	{
-		SetShaderValue(RHICmdList, GetPixelShader(), BufferSizeAndDirection, FVector4(InBufferSize, InDir));
+		SetShaderValue(RHICmdList, RHICmdList.GetBoundPixelShader(), BufferSizeAndDirection, FVector4(InBufferSize, InDir));
 	}
 
 	void SetWeightsAndOffsets(FRHICommandList& RHICmdList, const TArray<FVector4>& InWeightsAndOffsets, int32 NumSamples )
 	{
 		check(InWeightsAndOffsets.Num() <= MAX_BLUR_SAMPLES);
-		SetShaderValueArray<FRHIPixelShader*, FVector4>(RHICmdList, GetPixelShader(), WeightAndOffsets, InWeightsAndOffsets.GetData(), InWeightsAndOffsets.Num() );
-		SetShaderValue(RHICmdList, GetPixelShader(), SampleCount, NumSamples);
+		SetShaderValueArray<FRHIPixelShader*, FVector4>(RHICmdList, RHICmdList.GetBoundPixelShader(), WeightAndOffsets, InWeightsAndOffsets.GetData(), InWeightsAndOffsets.Num() );
+		SetShaderValue(RHICmdList, RHICmdList.GetBoundPixelShader(), SampleCount, NumSamples);
 	}
 
 	void SetUVBounds(FRHICommandList& RHICmdList, const FVector4& InUVBounds)
 	{
-		SetShaderValue(RHICmdList, GetPixelShader(), UVBounds, InUVBounds);
+		SetShaderValue(RHICmdList, RHICmdList.GetBoundPixelShader(), UVBounds, InUVBounds);
 	}
 
-	
-	virtual bool Serialize(FArchive& Ar) override
-	{
-		Ar << BufferSizeAndDirection;
-		Ar << WeightAndOffsets;
-		Ar << SampleCount;	
-		Ar << UVBounds;
-		return FSlateElementPS::Serialize(Ar);
-	}
 private:
-	FShaderParameter BufferSizeAndDirection;
-	FShaderParameter WeightAndOffsets;
-	FShaderParameter SampleCount;
-	FShaderParameter UVBounds;
+	LAYOUT_FIELD(FShaderParameter, BufferSizeAndDirection);
+	LAYOUT_FIELD(FShaderParameter, WeightAndOffsets);
+	LAYOUT_FIELD(FShaderParameter, SampleCount);
+	LAYOUT_FIELD(FShaderParameter, UVBounds);
 };
 
 
@@ -388,16 +346,11 @@ public:
 
 	void SetUVBounds(FRHICommandList& RHICmdList, const FVector4& InUVBounds)
 	{
-		SetShaderValue(RHICmdList, GetPixelShader(), UVBounds, InUVBounds);
+		SetShaderValue(RHICmdList, RHICmdList.GetBoundPixelShader(), UVBounds, InUVBounds);
 	}
 
-	virtual bool Serialize(FArchive& Ar) override
-	{
-		Ar << UVBounds;
-		return FSlateElementPS::Serialize(Ar);
-	}
 private:
-	FShaderParameter UVBounds;
+	LAYOUT_FIELD(FShaderParameter, UVBounds);
 };
 
 
@@ -427,29 +380,21 @@ public:
 
 	void SetColorRules(FRHICommandList& RHICmdList, bool bCorrect, EColorVisionDeficiency DeficiencyType, int32 Severity)
 	{
-		SetShaderValue(RHICmdList, GetPixelShader(), ColorVisionDeficiencyType, (float)DeficiencyType);
-		SetShaderValue(RHICmdList, GetPixelShader(), ColorVisionDeficiencySeverity, (float)Severity);
-		SetShaderValue(RHICmdList, GetPixelShader(), bCorrectDeficiency, bCorrect ? 1.0f : 0.0f);
+		SetShaderValue(RHICmdList, RHICmdList.GetBoundPixelShader(), ColorVisionDeficiencyType, (float)DeficiencyType);
+		SetShaderValue(RHICmdList, RHICmdList.GetBoundPixelShader(), ColorVisionDeficiencySeverity, (float)Severity);
+		SetShaderValue(RHICmdList, RHICmdList.GetBoundPixelShader(), bCorrectDeficiency, bCorrect ? 1.0f : 0.0f);
 	}
 
 	void SetShowCorrectionWithDeficiency(FRHICommandList& RHICmdList, bool bShowCorrectionWithDeficiency)
 	{
-		SetShaderValue(RHICmdList, GetPixelShader(), bSimulateCorrectionWithDeficiency, bShowCorrectionWithDeficiency ? 1.0f : 0.0f);
+		SetShaderValue(RHICmdList, RHICmdList.GetBoundPixelShader(), bSimulateCorrectionWithDeficiency, bShowCorrectionWithDeficiency ? 1.0f : 0.0f);
 	}
 
-	virtual bool Serialize(FArchive& Ar) override
-	{
-		Ar << ColorVisionDeficiencyType;
-		Ar << ColorVisionDeficiencySeverity;
-		Ar << bCorrectDeficiency;
-		Ar << bSimulateCorrectionWithDeficiency;
-		return FSlateElementPS::Serialize(Ar);
-	}
 private:
-	FShaderParameter ColorVisionDeficiencyType;
-	FShaderParameter ColorVisionDeficiencySeverity;
-	FShaderParameter bCorrectDeficiency;
-	FShaderParameter bSimulateCorrectionWithDeficiency;
+	LAYOUT_FIELD(FShaderParameter, ColorVisionDeficiencyType);
+	LAYOUT_FIELD(FShaderParameter, ColorVisionDeficiencySeverity);
+	LAYOUT_FIELD(FShaderParameter, bCorrectDeficiency);
+	LAYOUT_FIELD(FShaderParameter, bSimulateCorrectionWithDeficiency);
 };
 
 
@@ -487,15 +432,15 @@ public:
 	 */
 	void SetMaskRect(FRHICommandList& RHICmdList, const FVector2D& TopLeft, const FVector2D& TopRight, const FVector2D& BotLeft, const FVector2D& BotRight);
 
-	virtual bool Serialize(FArchive& Ar) override;
+	//virtual bool Serialize(FArchive& Ar) override;
 
 private:
 	/** Mask rect parameter */
-	FShaderParameter MaskRect;
+	LAYOUT_FIELD(FShaderParameter, MaskRect)
 	/** ViewProjection parameter used by the shader */
-	FShaderParameter ViewProjection;
+	LAYOUT_FIELD(FShaderParameter, ViewProjection)
 	/** Parameter used to determine if we need to swtich the vertical axis for opengl */
-	FShaderParameter SwitchVerticalAxisMultiplier;
+	LAYOUT_FIELD(FShaderParameter, SwitchVerticalAxisMultiplier)
 };
 
 class FSlateMaskingPS : public FGlobalShader
@@ -517,11 +462,6 @@ public:
 		: FGlobalShader(Initializer)
 	{
 	}
-
-	virtual bool Serialize(FArchive& Ar) override
-	{
-		return FGlobalShader::Serialize(Ar);
-	}
 };
 
 
@@ -550,10 +490,10 @@ public:
 	{
 		static const auto CVarOutputDevice = IConsoleManager::Get().FindTConsoleVariableDataInt(TEXT("r.HDR.Display.OutputDevice"));
 
-		SetTextureParameter(RHICmdList, GetPixelShader(), SceneTexture, SceneSampler, TStaticSamplerState<SF_Point>::GetRHI(), SceneTextureRHI);
+		SetTextureParameter(RHICmdList, RHICmdList.GetBoundPixelShader(), SceneTexture, SceneSampler, TStaticSamplerState<SF_Point>::GetRHI(), SceneTextureRHI);
 		
 		static auto CVarHDRNITLevel = IConsoleManager::Get().FindConsoleVariable(TEXT("Editor.HDRNITLevel"));
-		SetShaderValue(RHICmdList, GetPixelShader(), UILevel, CVarHDRNITLevel->GetFloat());
+		SetShaderValue(RHICmdList, RHICmdList.GetBoundPixelShader(), UILevel, CVarHDRNITLevel->GetFloat());
 	}
 
 	static void ModifyCompilationEnvironment(const FGlobalShaderPermutationParameters& Parameters, FShaderCompilerEnvironment& OutEnvironment)
@@ -561,13 +501,6 @@ public:
 		FGlobalShader::ModifyCompilationEnvironment(Parameters, OutEnvironment);
 	}
 
-	virtual bool Serialize(FArchive& Ar) override
-	{
-		bool bShaderHasOutdatedParameters = FGlobalShader::Serialize(Ar);
-		Ar << SceneTexture << SceneSampler << UILevel;
-		return bShaderHasOutdatedParameters;
-	}
-
 	static const TCHAR* GetSourceFilename()
 	{
 		return TEXT("/Engine/Private/CompositeUIPixelShader.usf");
@@ -579,9 +512,9 @@ public:
 	}
 
 private:
-	FShaderResourceParameter SceneTexture;
-	FShaderResourceParameter SceneSampler;
-	FShaderParameter UILevel;
+	LAYOUT_FIELD(FShaderResourceParameter, SceneTexture);
+	LAYOUT_FIELD(FShaderResourceParameter, SceneSampler);
+	LAYOUT_FIELD(FShaderParameter, UILevel);
 };
 #endif
 
diff --git a/Engine/Source/Runtime/VulkanRHI/Private/VulkanCommands.cpp b/Engine/Source/Runtime/VulkanRHI/Private/VulkanCommands.cpp
--- a/Engine/Source/Runtime/VulkanRHI/Private/VulkanCommands.cpp
+++ b/Engine/Source/Runtime/VulkanRHI/Private/VulkanCommands.cpp
@@ -50,12 +50,12 @@ static_assert(STRUCT_OFFSET(FRHIDispatchIndirectParameters, ThreadGroupCountZ) =
 
 static FORCEINLINE ShaderStage::EStage GetAndVerifyShaderStage(FRHIGraphicsShader* ShaderRHI, FVulkanPendingGfxState* PendingGfxState)
 {
-	switch (ShaderRHI->GetType())
+	switch (ShaderRHI->GetFrequency())
 	{
-	case FRHIShader::EType::Vertex:
+	case SF_Vertex:
 		check(PendingGfxState->GetCurrentShaderKey(ShaderStage::Vertex) == GetShaderKey<FVulkanVertexShader>(ShaderRHI));
 		return ShaderStage::Vertex;
-	case FRHIShader::EType::Hull:
+	case SF_Hull:
 #if PLATFORM_SUPPORTS_TESSELLATION_SHADERS
 		check(PendingGfxState->GetCurrentShaderKey(ShaderStage::Hull) == GetShaderKey<FVulkanHullShader>(ShaderRHI));
 		return ShaderStage::Hull;
@@ -64,7 +64,7 @@ static FORCEINLINE ShaderStage::EStage GetAndVerifyShaderStage(FRHIGraphicsShade
 		UE_LOG(LogVulkanRHI, Fatal, TEXT("Tessellation (Hull) not supported on this platform!"));
 		break;
 #endif
-	case FRHIShader::EType::Domain:
+	case SF_Domain:
 #if PLATFORM_SUPPORTS_TESSELLATION_SHADERS
 		check(PendingGfxState->GetCurrentShaderKey(ShaderStage::Domain) == GetShaderKey<FVulkanDomainShader>(ShaderRHI));
 		return ShaderStage::Domain;
@@ -73,7 +73,7 @@ static FORCEINLINE ShaderStage::EStage GetAndVerifyShaderStage(FRHIGraphicsShade
 		UE_LOG(LogVulkanRHI, Fatal, TEXT("Tessellation (Domain) not supported on this platform!"));
 		break;
 #endif
-	case FRHIShader::EType::Geometry:
+	case SF_Geometry:
 #if VULKAN_SUPPORTS_GEOMETRY_SHADERS
 		check(PendingGfxState->GetCurrentShaderKey(ShaderStage::Geometry) == GetShaderKey<FVulkanGeometryShader>(ShaderRHI));
 		return ShaderStage::Geometry;
@@ -82,11 +82,11 @@ static FORCEINLINE ShaderStage::EStage GetAndVerifyShaderStage(FRHIGraphicsShade
 		UE_LOG(LogVulkanRHI, Fatal, TEXT("Geometry shaders not supported on this platform!"));
 		break;
 #endif
-	case FRHIShader::EType::Pixel:
+	case SF_Pixel:
 		check(PendingGfxState->GetCurrentShaderKey(ShaderStage::Pixel) == GetShaderKey<FVulkanPixelShader>(ShaderRHI));
 		return ShaderStage::Pixel;
 	default:
-		checkf(0, TEXT("Undefined FRHIShader Type %d!"), (int32)ShaderRHI->GetType());
+		checkf(0, TEXT("Undefined FRHIShader Type %d!"), (int32)ShaderRHI->GetFrequency());
 		break;
 	}
 
@@ -95,13 +95,13 @@ static FORCEINLINE ShaderStage::EStage GetAndVerifyShaderStage(FRHIGraphicsShade
 
 static FORCEINLINE ShaderStage::EStage GetAndVerifyShaderStageAndVulkanShader(FRHIGraphicsShader* ShaderRHI, FVulkanPendingGfxState* PendingGfxState, FVulkanShader*& OutShader)
 {
-	switch (ShaderRHI->GetType())
+	switch (ShaderRHI->GetFrequency())
 	{
-	case FRHIShader::EType::Vertex:
+	case SF_Vertex:
 		//check(PendingGfxState->GetCurrentShaderKey(ShaderStage::Vertex) == GetShaderKey<FVulkanVertexShader>(ShaderRHI));
 		OutShader = static_cast<FVulkanVertexShader*>(static_cast<FRHIVertexShader*>(ShaderRHI));
 		return ShaderStage::Vertex;
-	case FRHIShader::EType::Hull:
+	case SF_Hull:
 #if PLATFORM_SUPPORTS_TESSELLATION_SHADERS
 		//check(PendingGfxState->GetCurrentShaderKey(ShaderStage::Hull) == GetShaderKey<FVulkanHullShader>(ShaderRHI));
 		OutShader = static_cast<FVulkanHullShader*>(static_cast<FRHIHullShader*>(ShaderRHI));
@@ -111,7 +111,7 @@ static FORCEINLINE ShaderStage::EStage GetAndVerifyShaderStageAndVulkanShader(FR
 		UE_LOG(LogVulkanRHI, Fatal, TEXT("Tessellation (Domain) not supported on this platform!"));
 		break;
 #endif
-	case FRHIShader::EType::Domain:
+	case SF_Domain:
 #if PLATFORM_SUPPORTS_TESSELLATION_SHADERS
 		//check(PendingGfxState->GetCurrentShaderKey(ShaderStage::Domain) == GetShaderKey<FVulkanDomainShader>(ShaderRHI));
 		OutShader = static_cast<FVulkanDomainShader*>(static_cast<FRHIDomainShader*>(ShaderRHI));
@@ -121,7 +121,7 @@ static FORCEINLINE ShaderStage::EStage GetAndVerifyShaderStageAndVulkanShader(FR
 		UE_LOG(LogVulkanRHI, Fatal, TEXT("Tessellation (Hull) not supported on this platform!"));
 		break;
 #endif
-	case FRHIShader::EType::Geometry:
+	case SF_Geometry:
 #if VULKAN_SUPPORTS_GEOMETRY_SHADERS
 		//check(PendingGfxState->GetCurrentShaderKey(ShaderStage::Geometry) == GetShaderKey<FVulkanGeometryShader>(ShaderRHI));
 		OutShader = static_cast<FVulkanGeometryShader*>(static_cast<FRHIGeometryShader*>(ShaderRHI));
@@ -131,12 +131,12 @@ static FORCEINLINE ShaderStage::EStage GetAndVerifyShaderStageAndVulkanShader(FR
 		UE_LOG(LogVulkanRHI, Fatal, TEXT("Geometry shaders not supported on this platform!"));
 		break;
 #endif
-	case FRHIShader::EType::Pixel:
+	case SF_Pixel:
 		//check(PendingGfxState->GetCurrentShaderKey(ShaderStage::Pixel) == GetShaderKey<FVulkanPixelShader>(ShaderRHI));
 		OutShader = static_cast<FVulkanPixelShader*>(static_cast<FRHIPixelShader*>(ShaderRHI));
 		return ShaderStage::Pixel;
 	default:
-		checkf(0, TEXT("Undefined FRHIShader Type %d!"), (int32)ShaderRHI->GetType());
+		checkf(0, TEXT("Undefined FRHIShader Type %d!"), (int32)ShaderRHI->GetFrequency());
 		break;
 	}
 
diff --git a/Engine/Source/Runtime/VulkanRHI/Private/VulkanPipeline.cpp b/Engine/Source/Runtime/VulkanRHI/Private/VulkanPipeline.cpp
--- a/Engine/Source/Runtime/VulkanRHI/Private/VulkanPipeline.cpp
+++ b/Engine/Source/Runtime/VulkanRHI/Private/VulkanPipeline.cpp
@@ -1045,7 +1045,7 @@ bool FVulkanPipelineStateCacheManager::CreateGfxPipelineFromEntry(FVulkanRHIGrap
 	FGfxPipelineDesc* GfxEntry = &PSO->Desc;
 	if (Shaders[ShaderStage::Pixel] == nullptr && !FVulkanPlatform::SupportsNullPixelShader())
 	{
-		Shaders[ShaderStage::Pixel] = ResourceCast(TShaderMapRef<FNULLPS>(GetGlobalShaderMap(GMaxRHIFeatureLevel))->GetPixelShader());
+		Shaders[ShaderStage::Pixel] = ResourceCast(TShaderMapRef<FNULLPS>(GetGlobalShaderMap(GMaxRHIFeatureLevel)).GetPixelShader());
 	}
 
 	if (!PSO->bShaderModulesLoaded)
diff --git a/Engine/Source/Runtime/VulkanRHI/Private/VulkanRHI.cpp b/Engine/Source/Runtime/VulkanRHI/Private/VulkanRHI.cpp
--- a/Engine/Source/Runtime/VulkanRHI/Private/VulkanRHI.cpp
+++ b/Engine/Source/Runtime/VulkanRHI/Private/VulkanRHI.cpp
@@ -254,7 +254,7 @@ void FVulkanDynamicRHI::Init()
 void FVulkanDynamicRHI::PostInit()
 {
 	//work around layering violation
-	TShaderMapRef<FNULLPS>(GetGlobalShaderMap(GMaxRHIFeatureLevel))->GetPixelShader();
+	TShaderMapRef<FNULLPS>(GetGlobalShaderMap(GMaxRHIFeatureLevel)).GetPixelShader();
 }
 
 void FVulkanDynamicRHI::Shutdown()
@@ -285,17 +285,7 @@ void FVulkanDynamicRHI::Shutdown()
 		check(!GIsCriticalError);
 
 		// Ask all initialized FRenderResources to release their RHI resources.
-		for (TLinkedList<FRenderResource*>::TIterator ResourceIt(FRenderResource::GetResourceList());ResourceIt;ResourceIt.Next())
-		{
-			FRenderResource* Resource = *ResourceIt;
-			check(Resource->IsInitialized());
-			Resource->ReleaseRHI();
-		}
-
-		for (TLinkedList<FRenderResource*>::TIterator ResourceIt(FRenderResource::GetResourceList());ResourceIt;ResourceIt.Next())
-		{
-			ResourceIt->ReleaseDynamicRHI();
-		}
+		FRenderResource::ReleaseRHIForAllResources();
 
 		{
 			for (auto& Pair : Device->SamplerMap)
diff --git a/Engine/Source/Runtime/VulkanRHI/Private/VulkanResources.h b/Engine/Source/Runtime/VulkanRHI/Private/VulkanResources.h
--- a/Engine/Source/Runtime/VulkanRHI/Private/VulkanResources.h
+++ b/Engine/Source/Runtime/VulkanRHI/Private/VulkanResources.h
@@ -88,7 +88,7 @@ public:
 
 	void PurgeShaderModules();
 
-	void Setup(const TArray<uint8>& InShaderHeaderAndCode, uint64 InShaderKey);
+	void Setup(TArrayView<const uint8> InShaderHeaderAndCode, uint64 InShaderKey);
 
 	VkShaderModule GetOrCreateHandle(const FVulkanLayout* Layout, uint32 LayoutHash)
 	{
@@ -192,7 +192,7 @@ public:
 	~FVulkanShaderFactory();
 	
 	template <typename ShaderType> 
-	ShaderType* CreateShader(const TArray<uint8>& Code, FVulkanDevice* Device);
+	ShaderType* CreateShader(TArrayView<const uint8> Code, FVulkanDevice* Device);
 	
 	template <typename ShaderType> 
 	ShaderType* LookupShader(uint64 ShaderKey) const
diff --git a/Engine/Source/Runtime/VulkanRHI/Private/VulkanShaders.cpp b/Engine/Source/Runtime/VulkanRHI/Private/VulkanShaders.cpp
--- a/Engine/Source/Runtime/VulkanRHI/Private/VulkanShaders.cpp
+++ b/Engine/Source/Runtime/VulkanRHI/Private/VulkanShaders.cpp
@@ -39,7 +39,7 @@ FVulkanShaderFactory::~FVulkanShaderFactory()
 }
 
 template <typename ShaderType> 
-ShaderType* FVulkanShaderFactory::CreateShader(const TArray<uint8>& Code, FVulkanDevice* Device)
+ShaderType* FVulkanShaderFactory::CreateShader(TArrayView<const uint8> Code, FVulkanDevice* Device)
 {
 	uint32 ShaderCodeLen = Code.Num();
 	uint32 ShaderCodeCRC = FCrc::MemCrc32(Code.GetData(), Code.Num());
@@ -83,14 +83,14 @@ void FVulkanShaderFactory::OnDeleteShader(const FVulkanShader& Shader)
 	ShaderMap[Shader.Frequency].Remove(ShaderKey);
 }
 
-void FVulkanShader::Setup(const TArray<uint8>& InShaderHeaderAndCode, uint64 InShaderKey)
+void FVulkanShader::Setup(TArrayView<const uint8> InShaderHeaderAndCode, uint64 InShaderKey)
 {
 	LLM_SCOPE_VULKAN(ELLMTagVulkan::VulkanShaders);
 	check(Device);
 
 	ShaderKey = InShaderKey;
 
-	FMemoryReader Ar(InShaderHeaderAndCode, true);
+	FMemoryReaderView Ar(InShaderHeaderAndCode, true);
 
 	Ar << CodeHeader;
 
@@ -224,32 +224,32 @@ VkShaderModule FVulkanLayout::CreatePatchedPatchSpirvModule(TArray<uint32>& Spir
 }
 
 
-FVertexShaderRHIRef FVulkanDynamicRHI::RHICreateVertexShader(const TArray<uint8>& Code)
+FVertexShaderRHIRef FVulkanDynamicRHI::RHICreateVertexShader(TArrayView<const uint8> Code, const FSHAHash& Hash)
 {
 	return Device->GetShaderFactory().CreateShader<FVulkanVertexShader>(Code, Device);
 }
 
-FPixelShaderRHIRef FVulkanDynamicRHI::RHICreatePixelShader(const TArray<uint8>& Code)
+FPixelShaderRHIRef FVulkanDynamicRHI::RHICreatePixelShader(TArrayView<const uint8> Code, const FSHAHash& Hash)
 {
 	return Device->GetShaderFactory().CreateShader<FVulkanPixelShader>(Code, Device);
 }
 
-FHullShaderRHIRef FVulkanDynamicRHI::RHICreateHullShader(const TArray<uint8>& Code) 
+FHullShaderRHIRef FVulkanDynamicRHI::RHICreateHullShader(TArrayView<const uint8> Code, const FSHAHash& Hash)
 { 
 	return Device->GetShaderFactory().CreateShader<FVulkanHullShader>(Code, Device);
 }
 
-FDomainShaderRHIRef FVulkanDynamicRHI::RHICreateDomainShader(const TArray<uint8>& Code) 
+FDomainShaderRHIRef FVulkanDynamicRHI::RHICreateDomainShader(TArrayView<const uint8> Code, const FSHAHash& Hash)
 { 
 	return Device->GetShaderFactory().CreateShader<FVulkanDomainShader>(Code, Device);
 }
 
-FGeometryShaderRHIRef FVulkanDynamicRHI::RHICreateGeometryShader(const TArray<uint8>& Code) 
+FGeometryShaderRHIRef FVulkanDynamicRHI::RHICreateGeometryShader(TArrayView<const uint8> Code, const FSHAHash& Hash)
 { 
 	return Device->GetShaderFactory().CreateShader<FVulkanGeometryShader>(Code, Device);
 }
 
-FComputeShaderRHIRef FVulkanDynamicRHI::RHICreateComputeShader(const TArray<uint8>& Code) 
+FComputeShaderRHIRef FVulkanDynamicRHI::RHICreateComputeShader(TArrayView<const uint8> Code, const FSHAHash& Hash)
 { 
 	return Device->GetShaderFactory().CreateShader<FVulkanComputeShader>(Code, Device);
 }
diff --git a/Engine/Source/Runtime/VulkanRHI/Private/VulkanUniformBuffer.cpp b/Engine/Source/Runtime/VulkanRHI/Private/VulkanUniformBuffer.cpp
--- a/Engine/Source/Runtime/VulkanRHI/Private/VulkanUniformBuffer.cpp
+++ b/Engine/Source/Runtime/VulkanRHI/Private/VulkanUniformBuffer.cpp
@@ -40,7 +40,7 @@ static void ValidateUniformBufferResource(const FRHIUniformBufferLayout& InLayou
 		&& (InLayout.Resources[Index].MemberType == UBMT_SRV || InLayout.Resources[Index].MemberType == UBMT_RDG_TEXTURE_SRV || InLayout.Resources[Index].MemberType == UBMT_RDG_BUFFER_SRV))
 		&& Validation == EUniformBufferValidation::ValidateResources)
 	{
-		checkf(Resource, TEXT("Invalid resource entry creating uniform buffer, %s.Resources[%u], ResourceType 0x%x."), *InLayout.GetDebugName().ToString(), Index, (uint8)InLayout.Resources[Index].MemberType);
+		checkf(Resource, TEXT("Invalid resource entry creating uniform buffer, %s.Resources[%u], ResourceType 0x%x."), *InLayout.GetDebugName(), Index, (uint8)InLayout.Resources[Index].MemberType);
 	}
 }
 
diff --git a/Engine/Source/Runtime/VulkanRHI/Private/VulkanViewport.cpp b/Engine/Source/Runtime/VulkanRHI/Private/VulkanViewport.cpp
--- a/Engine/Source/Runtime/VulkanRHI/Private/VulkanViewport.cpp
+++ b/Engine/Source/Runtime/VulkanRHI/Private/VulkanViewport.cpp
@@ -1018,7 +1018,7 @@ void FVulkanDynamicRHI::RHITick(float DeltaTime)
 			if (bRequested)
 			{
 				//work around layering violation
-				TShaderMapRef<FNULLPS>(GetGlobalShaderMap(GMaxRHIFeatureLevel))->GetPixelShader();
+				TShaderMapRef<FNULLPS>(GetGlobalShaderMap(GMaxRHIFeatureLevel)).GetPixelShader();
 			}
 
 			VulkanDevice->GetImmediateContext().GetTempFrameAllocationBuffer().Reset();
diff --git a/Engine/Source/Runtime/VulkanRHI/Public/VulkanDynamicRHI.h b/Engine/Source/Runtime/VulkanRHI/Public/VulkanDynamicRHI.h
--- a/Engine/Source/Runtime/VulkanRHI/Public/VulkanDynamicRHI.h
+++ b/Engine/Source/Runtime/VulkanRHI/Public/VulkanDynamicRHI.h
@@ -66,12 +66,12 @@ public:
 	virtual FDepthStencilStateRHIRef RHICreateDepthStencilState(const FDepthStencilStateInitializerRHI& Initializer) final override;
 	virtual FBlendStateRHIRef RHICreateBlendState(const FBlendStateInitializerRHI& Initializer) final override;
 	virtual FVertexDeclarationRHIRef RHICreateVertexDeclaration(const FVertexDeclarationElementList& Elements) final override;
-	virtual FPixelShaderRHIRef RHICreatePixelShader(const TArray<uint8>& Code) final override;
-	virtual FVertexShaderRHIRef RHICreateVertexShader(const TArray<uint8>& Code) final override;
-	virtual FHullShaderRHIRef RHICreateHullShader(const TArray<uint8>& Code) final override;
-	virtual FDomainShaderRHIRef RHICreateDomainShader(const TArray<uint8>& Code) final override;
-	virtual FGeometryShaderRHIRef RHICreateGeometryShader(const TArray<uint8>& Code) final override;
-	virtual FComputeShaderRHIRef RHICreateComputeShader(const TArray<uint8>& Code) final override;
+	virtual FPixelShaderRHIRef RHICreatePixelShader(TArrayView<const uint8> Code, const FSHAHash& Hash) final override;
+	virtual FVertexShaderRHIRef RHICreateVertexShader(TArrayView<const uint8> Code, const FSHAHash& Hash) final override;
+	virtual FHullShaderRHIRef RHICreateHullShader(TArrayView<const uint8> Code, const FSHAHash& Hash) final override;
+	virtual FDomainShaderRHIRef RHICreateDomainShader(TArrayView<const uint8> Code, const FSHAHash& Hash) final override;
+	virtual FGeometryShaderRHIRef RHICreateGeometryShader(TArrayView<const uint8> Code, const FSHAHash& Hash) final override;
+	virtual FComputeShaderRHIRef RHICreateComputeShader(TArrayView<const uint8> Code, const FSHAHash& Hash) final override;
 	virtual FComputeFenceRHIRef RHICreateComputeFence(const FName& Name) final override;
 	virtual FGPUFenceRHIRef RHICreateGPUFence(const FName &Name) final override;
 	virtual FStagingBufferRHIRef RHICreateStagingBuffer() final override;
@@ -213,37 +213,30 @@ public:
 	virtual void UnlockIndexBuffer_RenderThread(class FRHICommandListImmediate& RHICmdList, FRHIIndexBuffer* IndexBuffer) override final;
 #endif
 
-	virtual FVertexShaderRHIRef CreateVertexShader_RenderThread(class FRHICommandListImmediate& RHICmdList, const TArray<uint8>& Code) override final
+	virtual FVertexShaderRHIRef CreateVertexShader_RenderThread(class FRHICommandListImmediate& RHICmdList, TArrayView<const uint8> Code, const FSHAHash& Hash) override final
 	{
-		return RHICreateVertexShader(Code);
+		return RHICreateVertexShader(Code, Hash);
 	}
-	//virtual FVertexShaderRHIRef CreateVertexShader_RenderThread(class FRHICommandListImmediate& RHICmdList, FRHIShaderLibrary* Library, FSHAHash Hash) override final;
-	virtual FPixelShaderRHIRef CreatePixelShader_RenderThread(class FRHICommandListImmediate& RHICmdList, const TArray<uint8>& Code) override final
+	virtual FPixelShaderRHIRef CreatePixelShader_RenderThread(class FRHICommandListImmediate& RHICmdList, TArrayView<const uint8> Code, const FSHAHash& Hash) override final
 	{
-		return RHICreatePixelShader(Code);
+		return RHICreatePixelShader(Code, Hash);
 	}
-	//virtual FPixelShaderRHIRef CreatePixelShader_RenderThread(class FRHICommandListImmediate& RHICmdList, FRHIShaderLibrary* Library, FSHAHash Hash) override final;
-	virtual FGeometryShaderRHIRef CreateGeometryShader_RenderThread(class FRHICommandListImmediate& RHICmdList, const TArray<uint8>& Code) override final
+	virtual FGeometryShaderRHIRef CreateGeometryShader_RenderThread(class FRHICommandListImmediate& RHICmdList, TArrayView<const uint8> Code, const FSHAHash& Hash) override final
 	{
-		return RHICreateGeometryShader(Code);
+		return RHICreateGeometryShader(Code, Hash);
 	}
-	//virtual FGeometryShaderRHIRef CreateGeometryShader_RenderThread(class FRHICommandListImmediate& RHICmdList, FRHIShaderLibrary* Library, FSHAHash Hash) override final;
-
-	virtual FComputeShaderRHIRef CreateComputeShader_RenderThread(class FRHICommandListImmediate& RHICmdList, const TArray<uint8>& Code) override final
+	virtual FComputeShaderRHIRef CreateComputeShader_RenderThread(class FRHICommandListImmediate& RHICmdList, TArrayView<const uint8> Code, const FSHAHash& Hash) override final
 	{
-		return RHICreateComputeShader(Code);
+		return RHICreateComputeShader(Code, Hash);
 	}
-	//virtual FComputeShaderRHIRef CreateComputeShader_RenderThread(class FRHICommandListImmediate& RHICmdList, FRHIShaderLibrary* Library, FSHAHash Hash) override final;
-	virtual FHullShaderRHIRef CreateHullShader_RenderThread(class FRHICommandListImmediate& RHICmdList, const TArray<uint8>& Code) override final
+	virtual FHullShaderRHIRef CreateHullShader_RenderThread(class FRHICommandListImmediate& RHICmdList, TArrayView<const uint8> Code, const FSHAHash& Hash) override final
 	{
-		return RHICreateHullShader(Code);
+		return RHICreateHullShader(Code, Hash);
 	}
-	//virtual FHullShaderRHIRef CreateHullShader_RenderThread(class FRHICommandListImmediate& RHICmdList, FRHIShaderLibrary* Library, FSHAHash Hash) override final;
-	virtual FDomainShaderRHIRef CreateDomainShader_RenderThread(class FRHICommandListImmediate& RHICmdList, const TArray<uint8>& Code) override final
+	virtual FDomainShaderRHIRef CreateDomainShader_RenderThread(class FRHICommandListImmediate& RHICmdList, TArrayView<const uint8> Code, const FSHAHash& Hash) override final
 	{
-		return RHICreateDomainShader(Code);
+		return RHICreateDomainShader(Code, Hash);
 	}
-	//virtual FDomainShaderRHIRef CreateDomainShader_RenderThread(class FRHICommandListImmediate& RHICmdList, FRHIShaderLibrary* Library, FSHAHash Hash) override final;
 	virtual void* LockTexture2D_RenderThread(class FRHICommandListImmediate& RHICmdList, FRHITexture2D* Texture, uint32 MipIndex, EResourceLockMode LockMode, uint32& DestStride, bool bLockWithinMiptail, bool bNeedsDefaultRHIFlush = true) override final
 	{
 		return RHILockTexture2D(Texture, MipIndex, LockMode, DestStride, bLockWithinMiptail);
diff --git a/Engine/Source/Runtime/Windows/AudioMixerXAudio2/Private/AudioMixerPlatformXAudio2.cpp b/Engine/Source/Runtime/Windows/AudioMixerXAudio2/Private/AudioMixerPlatformXAudio2.cpp
--- a/Engine/Source/Runtime/Windows/AudioMixerXAudio2/Private/AudioMixerPlatformXAudio2.cpp
+++ b/Engine/Source/Runtime/Windows/AudioMixerXAudio2/Private/AudioMixerPlatformXAudio2.cpp
@@ -1238,7 +1238,7 @@ namespace Audio
 		static FName NAME_ADPCM(TEXT("ADPCM"));
 
 #if WITH_ENGINE
-		if (InSoundWave->IsStreaming())
+		if (InSoundWave->IsStreaming(nullptr))
 		{
 			if (InSoundWave->IsSeekableStreaming())
 			{
@@ -1277,7 +1277,7 @@ namespace Audio
 #if WITH_ENGINE
 		check(InSoundWave);
 
-		if (InSoundWave->IsStreaming())
+		if (InSoundWave->IsStreaming(nullptr))
 		{
 			if (InSoundWave->IsSeekableStreaming())
 			{
@@ -1292,7 +1292,7 @@ namespace Audio
 		}
 #endif
 
-		if (InSoundWave->IsStreaming())
+		if (InSoundWave->IsStreaming(nullptr))
 		{
 #if USE_VORBIS_FOR_STREAMING
 			return new FVorbisAudioInfo();
diff --git a/Engine/Source/Runtime/Windows/D3D11RHI/Private/D3D11Commands.cpp b/Engine/Source/Runtime/Windows/D3D11RHI/Private/D3D11Commands.cpp
--- a/Engine/Source/Runtime/Windows/D3D11RHI/Private/D3D11Commands.cpp
+++ b/Engine/Source/Runtime/Windows/D3D11RHI/Private/D3D11Commands.cpp
@@ -371,37 +371,37 @@ void FD3D11DynamicRHI::RHISetShaderTexture(FRHIGraphicsShader* ShaderRHI,uint32
 	FD3D11TextureBase* NewTexture = GetD3D11TextureFromRHITexture(NewTextureRHI);
 	ID3D11ShaderResourceView* ShaderResourceView = NewTexture ? NewTexture->GetShaderResourceView() : nullptr;
 
-	switch (ShaderRHI->GetType())
+	switch (ShaderRHI->GetFrequency())
 	{
-	case FRHIShader::EType::Vertex:
+	case SF_Vertex:
 	{
 		FD3D11VertexShader* VertexShader = static_cast<FD3D11VertexShader*>(ShaderRHI);
 		VALIDATE_BOUND_SHADER(VertexShader);
 		SetShaderTexture<SF_Vertex>(NewTexture, ShaderResourceView, TextureIndex, NewTextureRHI);
 	}
 	break;
-	case FRHIShader::EType::Hull:
+	case SF_Hull:
 	{
 		FD3D11HullShader* HullShader = static_cast<FD3D11HullShader*>(ShaderRHI);
 		VALIDATE_BOUND_SHADER(HullShader);
 		SetShaderTexture<SF_Hull>(NewTexture, ShaderResourceView, TextureIndex, NewTextureRHI);
 	}
 	break;
-	case FRHIShader::EType::Domain:
+	case SF_Domain:
 	{
 		FD3D11DomainShader* DomainShader = static_cast<FD3D11DomainShader*>(ShaderRHI);
 		VALIDATE_BOUND_SHADER(DomainShader);
 		SetShaderTexture<SF_Domain>(NewTexture, ShaderResourceView, TextureIndex, NewTextureRHI);
 	}
 	break;
-	case FRHIShader::EType::Geometry:
+	case SF_Geometry:
 	{
 		FD3D11GeometryShader* GeometryShader = static_cast<FD3D11GeometryShader*>(ShaderRHI);
 		VALIDATE_BOUND_SHADER(GeometryShader);
 		SetShaderTexture<SF_Geometry>(NewTexture, ShaderResourceView, TextureIndex, NewTextureRHI);
 	}
 	break;
-	case FRHIShader::EType::Pixel:
+	case SF_Pixel:
 	{
 		FD3D11PixelShader* PixelShader = static_cast<FD3D11PixelShader*>(ShaderRHI);
 		VALIDATE_BOUND_SHADER(PixelShader);
@@ -409,7 +409,7 @@ void FD3D11DynamicRHI::RHISetShaderTexture(FRHIGraphicsShader* ShaderRHI,uint32
 	}
 	break;
 	default:
-		checkf(0, TEXT("Undefined FRHIShader Type %d!"), (int32)ShaderRHI->GetType());
+		checkf(0, TEXT("Undefined FRHIShader Type %d!"), (int32)ShaderRHI->GetFrequency());
 	}
 }
 
@@ -504,37 +504,37 @@ void FD3D11DynamicRHI::RHISetShaderResourceViewParameter(FRHIGraphicsShader* Sha
 		Resource = SRV->Resource;
 		D3D11SRV = SRV->View;
 	}
-	switch (ShaderRHI->GetType())
+	switch (ShaderRHI->GetFrequency())
 	{
-	case FRHIShader::EType::Vertex:
+	case SF_Vertex:
 	{
 		FD3D11VertexShader* VertexShader = static_cast<FD3D11VertexShader*>(ShaderRHI);
 		VALIDATE_BOUND_SHADER(VertexShader);
 		SetShaderResourceView<SF_Vertex>(Resource, D3D11SRV, TextureIndex, NAME_None);
 	}
 	break;
-	case FRHIShader::EType::Hull:
+	case SF_Hull:
 	{
 		FD3D11HullShader* HullShader = static_cast<FD3D11HullShader*>(ShaderRHI);
 		VALIDATE_BOUND_SHADER(HullShader);
 		SetShaderResourceView<SF_Hull>(Resource, D3D11SRV, TextureIndex, NAME_None);
 	}
 	break;
-	case FRHIShader::EType::Domain:
+	case SF_Domain:
 	{
 		FD3D11DomainShader* DomainShader = static_cast<FD3D11DomainShader*>(ShaderRHI);
 		VALIDATE_BOUND_SHADER(DomainShader);
 		SetShaderResourceView<SF_Domain>(Resource, D3D11SRV, TextureIndex, NAME_None);
 	}
 	break;
-	case FRHIShader::EType::Geometry:
+	case SF_Geometry:
 	{
 		FD3D11GeometryShader* GeometryShader = static_cast<FD3D11GeometryShader*>(ShaderRHI);
 		VALIDATE_BOUND_SHADER(GeometryShader);
 		SetShaderResourceView<SF_Geometry>(Resource, D3D11SRV, TextureIndex, NAME_None);
 	}
 	break;
-	case FRHIShader::EType::Pixel:
+	case SF_Pixel:
 	{
 		FD3D11PixelShader* PixelShader = static_cast<FD3D11PixelShader*>(ShaderRHI);
 		VALIDATE_BOUND_SHADER(PixelShader);
@@ -542,7 +542,7 @@ void FD3D11DynamicRHI::RHISetShaderResourceViewParameter(FRHIGraphicsShader* Sha
 	}
 	break;
 	default:
-		checkf(0, TEXT("Undefined FRHIShader Type %d!"), (int32)ShaderRHI->GetType());
+		checkf(0, TEXT("Undefined FRHIShader Type %d!"), (int32)ShaderRHI->GetFrequency());
 	}
 }
 
@@ -568,37 +568,37 @@ void FD3D11DynamicRHI::RHISetShaderSampler(FRHIGraphicsShader* ShaderRHI,uint32
 {
 	FD3D11SamplerState* NewState = ResourceCast(NewStateRHI);
 	ID3D11SamplerState* StateResource = NewState->Resource;
-	switch (ShaderRHI->GetType())
+	switch (ShaderRHI->GetFrequency())
 	{
-	case FRHIShader::EType::Vertex:
+	case SF_Vertex:
 	{
 		FD3D11VertexShader* VertexShader = static_cast<FD3D11VertexShader*>(ShaderRHI);
 		VALIDATE_BOUND_SHADER(VertexShader);
 		StateCache.SetSamplerState<SF_Vertex>(StateResource, SamplerIndex);
 	}
 		break;
-	case FRHIShader::EType::Hull:
+	case SF_Hull:
 	{
 		FD3D11HullShader* HullShader = static_cast<FD3D11HullShader*>(ShaderRHI);
 		VALIDATE_BOUND_SHADER(HullShader);
 		StateCache.SetSamplerState<SF_Hull>(StateResource, SamplerIndex);
 	}
 	break;
-	case FRHIShader::EType::Domain:
+	case SF_Domain:
 	{
 		FD3D11DomainShader* DomainShader = static_cast<FD3D11DomainShader*>(ShaderRHI);
 		VALIDATE_BOUND_SHADER(DomainShader);
 		StateCache.SetSamplerState<SF_Domain>(StateResource, SamplerIndex);
 	}
 	break;
-	case FRHIShader::EType::Geometry:
+	case SF_Geometry:
 	{
 		FD3D11GeometryShader* GeometryShader = static_cast<FD3D11GeometryShader*>(ShaderRHI);
 		VALIDATE_BOUND_SHADER(GeometryShader);
 		StateCache.SetSamplerState<SF_Geometry>(StateResource, SamplerIndex);
 	}
 	break;
-	case FRHIShader::EType::Pixel:
+	case SF_Pixel:
 	{
 		FD3D11PixelShader* PixelShader = static_cast<FD3D11PixelShader*>(ShaderRHI);
 		VALIDATE_BOUND_SHADER(PixelShader);
@@ -606,7 +606,7 @@ void FD3D11DynamicRHI::RHISetShaderSampler(FRHIGraphicsShader* ShaderRHI,uint32
 	}
 	break;
 	default:
-		checkf(0, TEXT("Undefined FRHIShader Type %d!"), (int32)ShaderRHI->GetType());
+		checkf(0, TEXT("Undefined FRHIShader Type %d!"), (int32)ShaderRHI->GetFrequency());
 	}
 }
 
@@ -636,9 +636,9 @@ void FD3D11DynamicRHI::RHISetShaderUniformBuffer(FRHIGraphicsShader* ShaderRHI,u
 	FD3D11UniformBuffer* Buffer = ResourceCast(BufferRHI);
 	ID3D11Buffer* ConstantBuffer = Buffer ? Buffer->Resource : NULL;
 	EShaderFrequency Stage = SF_NumFrequencies;
-	switch (ShaderRHI->GetType())
+	switch (ShaderRHI->GetFrequency())
 	{
-	case FRHIShader::EType::Vertex:
+	case SF_Vertex:
 	{
 		FD3D11VertexShader* VertexShader = static_cast<FD3D11VertexShader*>(ShaderRHI);
 		VALIDATE_BOUND_SHADER(VertexShader);
@@ -646,7 +646,7 @@ void FD3D11DynamicRHI::RHISetShaderUniformBuffer(FRHIGraphicsShader* ShaderRHI,u
 		StateCache.SetConstantBuffer<SF_Vertex>(ConstantBuffer, BufferIndex);
 	}
 	break;
-	case FRHIShader::EType::Hull:
+	case SF_Hull:
 	{
 		FD3D11HullShader* HullShader = static_cast<FD3D11HullShader*>(ShaderRHI);
 		VALIDATE_BOUND_SHADER(HullShader);
@@ -654,7 +654,7 @@ void FD3D11DynamicRHI::RHISetShaderUniformBuffer(FRHIGraphicsShader* ShaderRHI,u
 		StateCache.SetConstantBuffer<SF_Hull>(ConstantBuffer, BufferIndex);
 	}
 	break;
-	case FRHIShader::EType::Domain:
+	case SF_Domain:
 	{
 		FD3D11DomainShader* DomainShader = static_cast<FD3D11DomainShader*>(ShaderRHI);
 		VALIDATE_BOUND_SHADER(DomainShader);
@@ -662,7 +662,7 @@ void FD3D11DynamicRHI::RHISetShaderUniformBuffer(FRHIGraphicsShader* ShaderRHI,u
 		StateCache.SetConstantBuffer<SF_Domain>(ConstantBuffer, BufferIndex);
 	}
 	break;
-	case FRHIShader::EType::Geometry:
+	case SF_Geometry:
 	{
 		FD3D11GeometryShader* GeometryShader = static_cast<FD3D11GeometryShader*>(ShaderRHI);
 		VALIDATE_BOUND_SHADER(GeometryShader);
@@ -670,7 +670,7 @@ void FD3D11DynamicRHI::RHISetShaderUniformBuffer(FRHIGraphicsShader* ShaderRHI,u
 		StateCache.SetConstantBuffer<SF_Geometry>(ConstantBuffer, BufferIndex);
 	}
 	break;
-	case FRHIShader::EType::Pixel:
+	case SF_Pixel:
 	{
 		FD3D11PixelShader* PixelShader = static_cast<FD3D11PixelShader*>(ShaderRHI);
 		VALIDATE_BOUND_SHADER(PixelShader);
@@ -679,7 +679,7 @@ void FD3D11DynamicRHI::RHISetShaderUniformBuffer(FRHIGraphicsShader* ShaderRHI,u
 	}
 	break;
 	default:
-		checkf(0, TEXT("Undefined FRHIShader Type %d!"), (int32)ShaderRHI->GetType());
+		checkf(0, TEXT("Undefined FRHIShader Type %d!"), (int32)ShaderRHI->GetFrequency());
 		return;
 	}
 
@@ -703,9 +703,9 @@ void FD3D11DynamicRHI::RHISetShaderUniformBuffer(FRHIComputeShader* ComputeShade
 
 void FD3D11DynamicRHI::RHISetShaderParameter(FRHIGraphicsShader* ShaderRHI,uint32 BufferIndex,uint32 BaseIndex,uint32 NumBytes,const void* NewValue)
 {
-	switch (ShaderRHI->GetType())
+	switch (ShaderRHI->GetFrequency())
 	{
-	case FRHIShader::EType::Vertex:
+	case SF_Vertex:
 	{
 		FD3D11VertexShader* VertexShader = static_cast<FD3D11VertexShader*>(ShaderRHI);
 		VALIDATE_BOUND_SHADER(VertexShader);
@@ -713,7 +713,7 @@ void FD3D11DynamicRHI::RHISetShaderParameter(FRHIGraphicsShader* ShaderRHI,uint3
 		VSConstantBuffers[BufferIndex]->UpdateConstant((const uint8*)NewValue, BaseIndex, NumBytes);
 	}
 	break;
-	case FRHIShader::EType::Hull:
+	case SF_Hull:
 	{
 		FD3D11HullShader* HullShader = static_cast<FD3D11HullShader*>(ShaderRHI);
 		VALIDATE_BOUND_SHADER(HullShader);
@@ -721,7 +721,7 @@ void FD3D11DynamicRHI::RHISetShaderParameter(FRHIGraphicsShader* ShaderRHI,uint3
 		HSConstantBuffers[BufferIndex]->UpdateConstant((const uint8*)NewValue, BaseIndex, NumBytes);
 	}
 	break;
-	case FRHIShader::EType::Domain:
+	case SF_Domain:
 	{
 		FD3D11DomainShader* DomainShader = static_cast<FD3D11DomainShader*>(ShaderRHI);
 		VALIDATE_BOUND_SHADER(DomainShader);
@@ -729,7 +729,7 @@ void FD3D11DynamicRHI::RHISetShaderParameter(FRHIGraphicsShader* ShaderRHI,uint3
 		DSConstantBuffers[BufferIndex]->UpdateConstant((const uint8*)NewValue, BaseIndex, NumBytes);
 	}
 	break;
-	case FRHIShader::EType::Geometry:
+	case SF_Geometry:
 	{
 		FD3D11GeometryShader* GeometryShader = static_cast<FD3D11GeometryShader*>(ShaderRHI);
 		VALIDATE_BOUND_SHADER(GeometryShader);
@@ -737,7 +737,7 @@ void FD3D11DynamicRHI::RHISetShaderParameter(FRHIGraphicsShader* ShaderRHI,uint3
 		GSConstantBuffers[BufferIndex]->UpdateConstant((const uint8*)NewValue, BaseIndex, NumBytes);
 	}
 	break;
-	case FRHIShader::EType::Pixel:
+	case SF_Pixel:
 	{
 		FD3D11PixelShader* PixelShader = static_cast<FD3D11PixelShader*>(ShaderRHI);
 		VALIDATE_BOUND_SHADER(PixelShader);
@@ -746,7 +746,7 @@ void FD3D11DynamicRHI::RHISetShaderParameter(FRHIGraphicsShader* ShaderRHI,uint3
 	}
 	break;
 	default:
-		checkf(0, TEXT("Undefined FRHIShader Type %d!"), (int32)ShaderRHI->GetType());
+		checkf(0, TEXT("Undefined FRHIShader Type %d!"), (int32)ShaderRHI->GetFrequency());
 	}
 }
 
@@ -1551,7 +1551,7 @@ void FD3D11DynamicRHI::SetResourcesFromTables(const ShaderType* RESTRICT Shader)
 			if (LayoutHash != Shader->ShaderResourceTable.ResourceTableLayoutHashes[BufferIndex])
 			{
 				auto& BufferLayout = Buffer->GetLayout();
-				FString DebugName = BufferLayout.GetDebugName().GetPlainNameString();
+				const auto& DebugName = BufferLayout.GetDebugName();
 				const FString& ShaderName = Shader->ShaderName;
 #if UE_BUILD_DEBUG
 				FString ShaderUB;
@@ -1579,7 +1579,7 @@ void FD3D11DynamicRHI::SetResourcesFromTables(const ShaderType* RESTRICT Shader)
 		}
 #endif
 
-		FName LayoutName = Buffer->GetLayout().GetDebugName();
+		FName LayoutName = *Buffer->GetLayout().GetDebugName();
 
 		// todo: could make this two pass: gather then set
 		SetShaderResourcesFromBuffer_Surface<(EShaderFrequency)ShaderType::StaticFrequency>(this, &StateCache, Buffer, Shader->ShaderResourceTable.TextureMap.GetData(), BufferIndex, LayoutName);
@@ -1605,7 +1605,7 @@ int32 FD3D11DynamicRHI::SetUAVPSResourcesFromTables(const ShaderType* RESTRICT S
 		FD3D11UniformBuffer* Buffer = (FD3D11UniformBuffer*)BoundUniformBuffers[ShaderType::StaticFrequency][BufferIndex].GetReference();
 
 		check(BufferIndex < Shader->ShaderResourceTable.ResourceTableLayoutHashes.Num());
-		FName LayoutName = Buffer->GetLayout().GetDebugName();
+		FName LayoutName = *Buffer->GetLayout().GetDebugName();
 
 		if ((EShaderFrequency)ShaderType::StaticFrequency == SF_Pixel)
 		{
diff --git a/Engine/Source/Runtime/Windows/D3D11RHI/Private/D3D11Device.cpp b/Engine/Source/Runtime/Windows/D3D11RHI/Private/D3D11Device.cpp
--- a/Engine/Source/Runtime/Windows/D3D11RHI/Private/D3D11Device.cpp
+++ b/Engine/Source/Runtime/Windows/D3D11RHI/Private/D3D11Device.cpp
@@ -483,17 +483,7 @@ void FD3D11DynamicRHI::CleanupD3DDevice()
 		check(!GIsCriticalError);
 
 		// Ask all initialized FRenderResources to release their RHI resources.
-		for(TLinkedList<FRenderResource*>::TIterator ResourceIt(FRenderResource::GetResourceList());ResourceIt;ResourceIt.Next())
-		{
-			FRenderResource* Resource = *ResourceIt;
-			check(Resource->IsInitialized());
-			Resource->ReleaseRHI();
-		}
-
-		for(TLinkedList<FRenderResource*>::TIterator ResourceIt(FRenderResource::GetResourceList());ResourceIt;ResourceIt.Next())
-		{
-			ResourceIt->ReleaseDynamicRHI();
-		}
+		FRenderResource::ReleaseRHIForAllResources();
 
 		extern void EmptyD3DSamplerStateCache();
 		EmptyD3DSamplerStateCache();
diff --git a/Engine/Source/Runtime/Windows/D3D11RHI/Private/D3D11RHIPrivate.h b/Engine/Source/Runtime/Windows/D3D11RHI/Private/D3D11RHIPrivate.h
--- a/Engine/Source/Runtime/Windows/D3D11RHI/Private/D3D11RHIPrivate.h
+++ b/Engine/Source/Runtime/Windows/D3D11RHI/Private/D3D11RHIPrivate.h
@@ -416,12 +416,12 @@ public:
 	virtual FDepthStencilStateRHIRef RHICreateDepthStencilState(const FDepthStencilStateInitializerRHI& Initializer) final override;
 	virtual FBlendStateRHIRef RHICreateBlendState(const FBlendStateInitializerRHI& Initializer) final override;
 	virtual FVertexDeclarationRHIRef RHICreateVertexDeclaration(const FVertexDeclarationElementList& Elements) final override;
-	virtual FPixelShaderRHIRef RHICreatePixelShader(const TArray<uint8>& Code) final override;
-	virtual FVertexShaderRHIRef RHICreateVertexShader(const TArray<uint8>& Code) final override;
-	virtual FHullShaderRHIRef RHICreateHullShader(const TArray<uint8>& Code) final override;
-	virtual FDomainShaderRHIRef RHICreateDomainShader(const TArray<uint8>& Code) final override;
-	virtual FGeometryShaderRHIRef RHICreateGeometryShader(const TArray<uint8>& Code) final override;
-	virtual FComputeShaderRHIRef RHICreateComputeShader(const TArray<uint8>& Code) final override;
+	virtual FPixelShaderRHIRef RHICreatePixelShader(TArrayView<const uint8> Code, const FSHAHash& Hash) final override;
+	virtual FVertexShaderRHIRef RHICreateVertexShader(TArrayView<const uint8> Code, const FSHAHash& Hash) final override;
+	virtual FHullShaderRHIRef RHICreateHullShader(TArrayView<const uint8> Code, const FSHAHash& Hash) final override;
+	virtual FDomainShaderRHIRef RHICreateDomainShader(TArrayView<const uint8> Code, const FSHAHash& Hash) final override;
+	virtual FGeometryShaderRHIRef RHICreateGeometryShader(TArrayView<const uint8> Code, const FSHAHash& Hash) final override;
+	virtual FComputeShaderRHIRef RHICreateComputeShader(TArrayView<const uint8> Code, const FSHAHash& Hash) final override;
 	virtual FStagingBufferRHIRef RHICreateStagingBuffer() final override;
 	virtual void RHICopyToStagingBuffer(FRHIVertexBuffer* SourceBuffer, FRHIStagingBuffer* DestinationStagingBuffer, uint32 Offset, uint32 NumBytes) final override;
 	virtual void RHIWriteGPUFence(FRHIGPUFence* Fence) final override;
@@ -614,12 +614,12 @@ public:
 	virtual ETextureReallocationStatus FinalizeAsyncReallocateTexture2D_RenderThread(class FRHICommandListImmediate& RHICmdList, FRHITexture2D* Texture2D, bool bBlockUntilCompleted) final override;
 	virtual ETextureReallocationStatus CancelAsyncReallocateTexture2D_RenderThread(class FRHICommandListImmediate& RHICmdList, FRHITexture2D* Texture2D, bool bBlockUntilCompleted) final override;
 	virtual FIndexBufferRHIRef CreateIndexBuffer_RenderThread(class FRHICommandListImmediate& RHICmdList, uint32 Stride, uint32 Size, uint32 InUsage, FRHIResourceCreateInfo& CreateInfo) final override;
-	virtual FVertexShaderRHIRef CreateVertexShader_RenderThread(class FRHICommandListImmediate& RHICmdList, const TArray<uint8>& Code) final override;
-	virtual FPixelShaderRHIRef CreatePixelShader_RenderThread(class FRHICommandListImmediate& RHICmdList, const TArray<uint8>& Code) final override;
-	virtual FGeometryShaderRHIRef CreateGeometryShader_RenderThread(class FRHICommandListImmediate& RHICmdList, const TArray<uint8>& Code) final override;
-	virtual FComputeShaderRHIRef CreateComputeShader_RenderThread(class FRHICommandListImmediate& RHICmdList, const TArray<uint8>& Code) final override;
-	virtual FHullShaderRHIRef CreateHullShader_RenderThread(class FRHICommandListImmediate& RHICmdList, const TArray<uint8>& Code) final override;
-	virtual FDomainShaderRHIRef CreateDomainShader_RenderThread(class FRHICommandListImmediate& RHICmdList, const TArray<uint8>& Code);
+	virtual FVertexShaderRHIRef CreateVertexShader_RenderThread(class FRHICommandListImmediate& RHICmdList, TArrayView<const uint8> Code, const FSHAHash& Hash) final override;
+	virtual FPixelShaderRHIRef CreatePixelShader_RenderThread(class FRHICommandListImmediate& RHICmdList, TArrayView<const uint8> Code, const FSHAHash& Hash) final override;
+	virtual FGeometryShaderRHIRef CreateGeometryShader_RenderThread(class FRHICommandListImmediate& RHICmdList, TArrayView<const uint8> Code, const FSHAHash& Hash) final override;
+	virtual FComputeShaderRHIRef CreateComputeShader_RenderThread(class FRHICommandListImmediate& RHICmdList, TArrayView<const uint8> Code, const FSHAHash& Hash) final override;
+	virtual FHullShaderRHIRef CreateHullShader_RenderThread(class FRHICommandListImmediate& RHICmdList, TArrayView<const uint8> Code, const FSHAHash& Hash) final override;
+	virtual FDomainShaderRHIRef CreateDomainShader_RenderThread(class FRHICommandListImmediate& RHICmdList, TArrayView<const uint8> Code, const FSHAHash& Hash);
 	virtual void* LockTexture2D_RenderThread(class FRHICommandListImmediate& RHICmdList, FRHITexture2D* Texture, uint32 MipIndex, EResourceLockMode LockMode, uint32& DestStride, bool bLockWithinMiptail, bool bNeedsDefaultRHIFlush = true) final override;
 	virtual void UnlockTexture2D_RenderThread(class FRHICommandListImmediate& RHICmdList, FRHITexture2D* Texture, uint32 MipIndex, bool bLockWithinMiptail, bool bNeedsDefaultRHIFlush = true) final override;
 	virtual void UpdateTexture2D_RenderThread(class FRHICommandListImmediate& RHICmdList, FRHITexture2D* Texture, uint32 MipIndex, const struct FUpdateTextureRegion2D& UpdateRegion, uint32 SourcePitch, const uint8* SourceData) final override;
diff --git a/Engine/Source/Runtime/Windows/D3D11RHI/Private/D3D11RenderTarget.cpp b/Engine/Source/Runtime/Windows/D3D11RHI/Private/D3D11RenderTarget.cpp
--- a/Engine/Source/Runtime/Windows/D3D11RHI/Private/D3D11RenderTarget.cpp
+++ b/Engine/Source/Runtime/Windows/D3D11RHI/Private/D3D11RenderTarget.cpp
@@ -132,8 +132,8 @@ void FD3D11DynamicRHI::ResolveTextureUsingShader(
 	TShaderMapRef<FResolveVS> ResolveVertexShader(ShaderMap);
 	TShaderMapRef<TPixelShader> ResolvePixelShader(ShaderMap);
 
-	GraphicsPSOInit.BoundShaderState.VertexShaderRHI = GETSAFERHISHADER_VERTEX(*ResolveVertexShader);
-	GraphicsPSOInit.BoundShaderState.PixelShaderRHI = GETSAFERHISHADER_PIXEL(*ResolvePixelShader);
+	GraphicsPSOInit.BoundShaderState.VertexShaderRHI = ResolveVertexShader.GetVertexShader();
+	GraphicsPSOInit.BoundShaderState.PixelShaderRHI = ResolvePixelShader.GetPixelShader();
 	GraphicsPSOInit.PrimitiveType = PT_TriangleStrip;
 
 	CurrentDepthTexture = DestTexture;
diff --git a/Engine/Source/Runtime/Windows/D3D11RHI/Private/D3D11Shaders.cpp b/Engine/Source/Runtime/Windows/D3D11RHI/Private/D3D11Shaders.cpp
--- a/Engine/Source/Runtime/Windows/D3D11RHI/Private/D3D11Shaders.cpp
+++ b/Engine/Source/Runtime/Windows/D3D11RHI/Private/D3D11Shaders.cpp
@@ -113,13 +113,13 @@ static inline void InitUniformBufferStaticSlots(TShaderType* Shader)
 	}
 }
 
-FVertexShaderRHIRef FD3D11DynamicRHI::RHICreateVertexShader(const TArray<uint8>& Code)
+FVertexShaderRHIRef FD3D11DynamicRHI::RHICreateVertexShader(TArrayView<const uint8> Code, const FSHAHash& Hash)
 {
 	FShaderCodeReader ShaderCode(Code);
 
 	FD3D11VertexShader* Shader = new FD3D11VertexShader;
 
-	FMemoryReader Ar( Code, true );
+	FMemoryReaderView Ar( Code, true );
 	Ar << Shader->ShaderResourceTable;
 	int32 Offset = Ar.Tell();
 	const uint8* CodePtr = Code.GetData() + Offset;
@@ -140,18 +140,18 @@ FVertexShaderRHIRef FD3D11DynamicRHI::RHICreateVertexShader(const TArray<uint8>&
 
 FVertexShaderRHIRef FD3D11DynamicRHI::CreateVertexShader_RenderThread(
 	class FRHICommandListImmediate& RHICmdList,
-	const TArray<uint8>& Code)
+	TArrayView<const uint8> Code, const FSHAHash& Hash)
 {
-	return RHICreateVertexShader(Code);
+	return RHICreateVertexShader(Code, Hash);
 }
 
-FGeometryShaderRHIRef FD3D11DynamicRHI::RHICreateGeometryShader(const TArray<uint8>& Code) 
+FGeometryShaderRHIRef FD3D11DynamicRHI::RHICreateGeometryShader(TArrayView<const uint8> Code, const FSHAHash& Hash)
 { 
 	FShaderCodeReader ShaderCode(Code);
 
 	FD3D11GeometryShader* Shader = new FD3D11GeometryShader;
 
-	FMemoryReader Ar( Code, true );
+	FMemoryReaderView Ar( Code, true );
 	Ar << Shader->ShaderResourceTable;
 	int32 Offset = Ar.Tell();
 	const uint8* CodePtr = Code.GetData() + Offset;
@@ -168,18 +168,18 @@ FGeometryShaderRHIRef FD3D11DynamicRHI::RHICreateGeometryShader(const TArray<uin
 
 FGeometryShaderRHIRef FD3D11DynamicRHI::CreateGeometryShader_RenderThread(
 	class FRHICommandListImmediate& RHICmdList,
-	const TArray<uint8>& Code)
+	TArrayView<const uint8> Code, const FSHAHash& Hash)
 {
-	return RHICreateGeometryShader(Code);
+	return RHICreateGeometryShader(Code, Hash);
 }
 
-FHullShaderRHIRef FD3D11DynamicRHI::RHICreateHullShader(const TArray<uint8>& Code) 
+FHullShaderRHIRef FD3D11DynamicRHI::RHICreateHullShader(TArrayView<const uint8> Code, const FSHAHash& Hash)
 { 
 	FShaderCodeReader ShaderCode(Code);
 
 	FD3D11HullShader* Shader = new FD3D11HullShader;
 
-	FMemoryReader Ar( Code, true );
+	FMemoryReaderView Ar( Code, true );
 	Ar << Shader->ShaderResourceTable;
 	int32 Offset = Ar.Tell();
 	const uint8* CodePtr = Code.GetData() + Offset;
@@ -196,18 +196,18 @@ FHullShaderRHIRef FD3D11DynamicRHI::RHICreateHullShader(const TArray<uint8>& Cod
 
 FHullShaderRHIRef FD3D11DynamicRHI::CreateHullShader_RenderThread(
 	class FRHICommandListImmediate& RHICmdList,
-	const TArray<uint8>& Code)
+	TArrayView<const uint8> Code, const FSHAHash& Hash)
 {
-	return RHICreateHullShader(Code);
+	return RHICreateHullShader(Code, Hash);
 }
 
-FDomainShaderRHIRef FD3D11DynamicRHI::RHICreateDomainShader(const TArray<uint8>& Code) 
+FDomainShaderRHIRef FD3D11DynamicRHI::RHICreateDomainShader(TArrayView<const uint8> Code, const FSHAHash& Hash)
 { 
 	FShaderCodeReader ShaderCode(Code);
 
 	FD3D11DomainShader* Shader = new FD3D11DomainShader;
 
-	FMemoryReader Ar(Code, true);
+	FMemoryReaderView Ar(Code, true);
 	Ar << Shader->ShaderResourceTable;
 	int32 Offset = Ar.Tell();
 	const uint8* CodePtr = Code.GetData() + Offset;
@@ -224,18 +224,18 @@ FDomainShaderRHIRef FD3D11DynamicRHI::RHICreateDomainShader(const TArray<uint8>&
 
 FDomainShaderRHIRef FD3D11DynamicRHI::CreateDomainShader_RenderThread(
 	class FRHICommandListImmediate& RHICmdList,
-	const TArray<uint8>& Code)
+	TArrayView<const uint8> Code, const FSHAHash& Hash)
 {
-	return RHICreateDomainShader(Code);
+	return RHICreateDomainShader(Code, Hash);
 }
 
-FPixelShaderRHIRef FD3D11DynamicRHI::RHICreatePixelShader(const TArray<uint8>& Code)
+FPixelShaderRHIRef FD3D11DynamicRHI::RHICreatePixelShader(TArrayView<const uint8> Code, const FSHAHash& Hash)
 {
 	FShaderCodeReader ShaderCode(Code);
 
 	FD3D11PixelShader* Shader = new FD3D11PixelShader;
 
-	FMemoryReader Ar( Code, true );
+	FMemoryReaderView Ar( Code, true );
 	Ar << Shader->ShaderResourceTable;
 	int32 Offset = Ar.Tell();
 	const uint8* CodePtr = Code.GetData() + Offset;
@@ -252,18 +252,18 @@ FPixelShaderRHIRef FD3D11DynamicRHI::RHICreatePixelShader(const TArray<uint8>& C
 
 FPixelShaderRHIRef FD3D11DynamicRHI::CreatePixelShader_RenderThread(
 	class FRHICommandListImmediate& RHICmdList,
-	const TArray<uint8>& Code)
+	TArrayView<const uint8> Code, const FSHAHash& Hash)
 {
-	return RHICreatePixelShader(Code);
+	return RHICreatePixelShader(Code, Hash);
 }
 
-FComputeShaderRHIRef FD3D11DynamicRHI::RHICreateComputeShader(const TArray<uint8>& Code) 
+FComputeShaderRHIRef FD3D11DynamicRHI::RHICreateComputeShader(TArrayView<const uint8> Code, const FSHAHash& Hash)
 { 
 	FShaderCodeReader ShaderCode(Code);
 
 	FD3D11ComputeShader* Shader = new FD3D11ComputeShader;
 
-	FMemoryReader Ar( Code, true );
+	FMemoryReaderView Ar( Code, true );
 	Ar << Shader->ShaderResourceTable;
 	int32 Offset = Ar.Tell();
 	const uint8* CodePtr = Code.GetData() + Offset;
@@ -280,9 +280,9 @@ FComputeShaderRHIRef FD3D11DynamicRHI::RHICreateComputeShader(const TArray<uint8
 
 FComputeShaderRHIRef FD3D11DynamicRHI::CreateComputeShader_RenderThread(
 	class FRHICommandListImmediate& RHICmdList,
-	const TArray<uint8>& Code)
+	TArrayView<const uint8> Code, const FSHAHash& Hash)
 {
-	return RHICreateComputeShader(Code);
+	return RHICreateComputeShader(Code, Hash);
 }
 
 void FD3D11DynamicRHI::RHISetMultipleViewports(uint32 Count, const FViewportBounds* Data) 
diff --git a/Engine/Source/Runtime/Windows/D3D11RHI/Private/D3D11UniformBuffer.cpp b/Engine/Source/Runtime/Windows/D3D11RHI/Private/D3D11UniformBuffer.cpp
--- a/Engine/Source/Runtime/Windows/D3D11RHI/Private/D3D11UniformBuffer.cpp
+++ b/Engine/Source/Runtime/Windows/D3D11RHI/Private/D3D11UniformBuffer.cpp
@@ -247,6 +247,7 @@ FUniformBufferRHIRef FD3D11DynamicRHI::RHICreateUniformBuffer(const void* Conten
 		NewUniformBuffer = new FD3D11UniformBuffer(this, Layout, nullptr, FRingAllocation());
 	}
 
+	const TCHAR* LayoutName = *Layout.GetDebugName();
 	if (Layout.Resources.Num())
 	{
 		int32 NumResources = Layout.Resources.Num();
@@ -255,16 +256,18 @@ FUniformBufferRHIRef FD3D11DynamicRHI::RHICreateUniformBuffer(const void* Conten
 
 		for (int32 i = 0; i < NumResources; ++i)
 		{
-			FRHIResource* Resource = *(FRHIResource**)((uint8*)Contents + Layout.Resources[i].MemberOffset);
+			const FRHIUniformBufferLayout::FResourceParameter& ResourceParam = Layout.Resources[i];
+			FRHIResource* Resource = *(FRHIResource**)((uint8*)Contents + ResourceParam.MemberOffset);
 
 			// Allow null SRV's in uniform buffers for feature levels that don't support SRV's in shaders
 			if (!(GMaxRHIFeatureLevel <= ERHIFeatureLevel::ES3_1 
 				&& (Layout.Resources[i].MemberType == UBMT_SRV || Layout.Resources[i].MemberType == UBMT_UAV || Layout.Resources[i].MemberType == UBMT_RDG_TEXTURE_SRV || Layout.Resources[i].MemberType == UBMT_RDG_BUFFER_SRV))
 				&& Validation == EUniformBufferValidation::ValidateResources)
 			{
-				checkf(Resource, TEXT("Invalid resource entry creating uniform buffer, %s.Resources[%u], ResourceType 0x%x."), *Layout.GetDebugName().ToString(), i, Layout.Resources[i].MemberType);
+				checkf(Resource, TEXT("Invalid resource entry creating uniform buffer, %s.Resources[%u], ResourceType 0x%x."), *Layout.GetDebugName(), i, Layout.Resources[i].MemberType);
 			}
 			NewUniformBuffer->ResourceTable[i] = Resource;
+			int a = 0;
 		}
 	}
 
@@ -316,7 +319,7 @@ void FD3D11DynamicRHI::RHIUpdateUniformBuffer(FRHIUniformBuffer* UniformBufferRH
 				&& UniformBufferValidation == EUniformBufferValidation::ValidateResources)
 			{
 				checkf(Resource, TEXT("Invalid resource entry creating uniform buffer, %s.Resources[%u], ResourceType 0x%x."),
-					*Layout.GetDebugName().ToString(),
+					*Layout.GetDebugName(),
 					ResourceIndex,
 					Layout.Resources[ResourceIndex].MemberType);
 			}
@@ -342,7 +345,7 @@ void FD3D11DynamicRHI::RHIUpdateUniformBuffer(FRHIUniformBuffer* UniformBufferRH
 					&& UniformBufferValidation == EUniformBufferValidation::ValidateResources)
 				{
 					checkf(Resource, TEXT("Invalid resource entry creating uniform buffer, %s.Resources[%u], ResourceType 0x%x."),
-						*Layout.GetDebugName().ToString(),
+						*Layout.GetDebugName(),
 						ResourceIndex,
 						Layout.Resources[ResourceIndex].MemberType);
 				}
diff --git a/Engine/Source/Runtime/Windows/D3D11RHI/Public/D3D11Util.h b/Engine/Source/Runtime/Windows/D3D11RHI/Public/D3D11Util.h
--- a/Engine/Source/Runtime/Windows/D3D11RHI/Public/D3D11Util.h
+++ b/Engine/Source/Runtime/Windows/D3D11RHI/Public/D3D11Util.h
@@ -244,12 +244,16 @@ private:
 	int32 NumActiveTargets;
 };
 
+struct FD3D11RHIGenericCommandString
+{
+	static const TCHAR* TStr() { return TEXT("FD3D11RHIGenericCommand"); }
+};
 template <
 	typename JobType,
 	typename = TEnableIf<TOr<
 	TIsSame<JobType, TFunction<void()>>,
 	TIsSame<JobType, TFunction<void()>&>>::Value>>
-	class TD3D11RHIGenericCommand final : public FRHICommand<TD3D11RHIGenericCommand<JobType>>
+	class TD3D11RHIGenericCommand final : public FRHICommand<TD3D11RHIGenericCommand<JobType>, FD3D11RHIGenericCommandString>
 {
 public:
 	// InRHIJob is supposed to be called on RHIT (don't capture things that can become outdated here)
diff --git a/Engine/Source/Runtime/Windows/XAudio2/Private/XAudio2Device.cpp b/Engine/Source/Runtime/Windows/XAudio2/Private/XAudio2Device.cpp
--- a/Engine/Source/Runtime/Windows/XAudio2/Private/XAudio2Device.cpp
+++ b/Engine/Source/Runtime/Windows/XAudio2/Private/XAudio2Device.cpp
@@ -360,7 +360,7 @@ class ICompressedAudioInfo* FXAudio2Device::CreateCompressedAudioInfo(USoundWave
 {
 	check(SoundWave);
 
-	if (SoundWave->IsStreaming())
+	if (SoundWave->IsStreaming(nullptr))
 	{
 		if (SoundWave->IsSeekableStreaming())
 		{
diff --git a/Engine/Source/Runtime/Windows/XAudio2/Public/XAudio2Device.h b/Engine/Source/Runtime/Windows/XAudio2/Public/XAudio2Device.h
--- a/Engine/Source/Runtime/Windows/XAudio2/Public/XAudio2Device.h
+++ b/Engine/Source/Runtime/Windows/XAudio2/Public/XAudio2Device.h
@@ -124,7 +124,7 @@ private:
 		static const FName NAME_OPUS(TEXT("OPUS"));
 		static const FName NAME_XMA(TEXT("XMA"));
 
-		if (SoundWave->IsStreaming())
+		if (SoundWave->IsStreaming(nullptr))
 		{
 			if (SoundWave->IsSeekableStreaming())
 			{
