Fix additive blended font materials causing the text to disappear

diff --git a/Engine/Shaders/Private/SlateElementPixelShader.usf b/Engine/Shaders/Private/SlateElementPixelShader.usf
--- a/Engine/Shaders/Private/SlateElementPixelShader.usf
+++ b/Engine/Shaders/Private/SlateElementPixelShader.usf
@@ -172,7 +172,11 @@ half4 GetFontElementColor( VertexToPixelInterpolants InVertex )
 	half4 OutColor = BaseColor;
 
 #if SHADER_TYPE == ST_GrayscaleFont
-	OutColor.a *= Texture2DSample_A8( ElementTexture, ElementTextureSampler, AtlasTextureCoordinates );
+	#if USE_MATERIALS && MATERIALBLENDING_ADDITIVE
+		OutColor *= Texture2DSample_A8( ElementTexture, ElementTextureSampler, AtlasTextureCoordinates );
+	#else
+		OutColor.a *= Texture2DSample_A8( ElementTexture, ElementTextureSampler, AtlasTextureCoordinates );
+	#endif
 #else
 	OutColor *= Texture2DSample( ElementTexture, ElementTextureSampler, AtlasTextureCoordinates );
 #endif
diff --git a/Engine/Source/Runtime/SlateRHIRenderer/Private/SlateRHIRenderingPolicy.cpp b/Engine/Source/Runtime/SlateRHIRenderer/Private/SlateRHIRenderingPolicy.cpp
--- a/Engine/Source/Runtime/SlateRHIRenderer/Private/SlateRHIRenderingPolicy.cpp
+++ b/Engine/Source/Runtime/SlateRHIRenderer/Private/SlateRHIRenderingPolicy.cpp
@@ -1054,10 +1054,6 @@ void FSlateRHIRenderingPolicy::DrawElements(
 						{
 							PixelShader->SetBlendState(GraphicsPSOInit, Material);
 							FSlateShaderResource* MaskResource = MaterialShaderResource->GetTextureMaskResource();
-							if (MaskResource)
-							{
-								GraphicsPSOInit.BlendState = TStaticBlendState<CW_RGBA, BO_Add, BF_SourceAlpha, BF_InverseSourceAlpha, BO_Add, BF_InverseDestAlpha, BF_One>::GetRHI();
-							}
 
 							GraphicsPSOInit.BoundShaderState.VertexDeclarationRHI = bUseInstancing ? GSlateInstancedVertexDeclaration.VertexDeclarationRHI : GSlateVertexDeclaration.VertexDeclarationRHI;
 							GraphicsPSOInit.BoundShaderState.VertexShaderRHI = VertexShader.GetVertexShader();
