Fix font materials incorrectly marked as opaque, not blending properly.  Font materials are forced into translucent blending to work properly

diff --git a/Engine/Source/Runtime/SlateRHIRenderer/Private/SlateRHIRenderingPolicy.cpp b/Engine/Source/Runtime/SlateRHIRenderer/Private/SlateRHIRenderingPolicy.cpp
--- a/Engine/Source/Runtime/SlateRHIRenderer/Private/SlateRHIRenderingPolicy.cpp
+++ b/Engine/Source/Runtime/SlateRHIRenderer/Private/SlateRHIRenderingPolicy.cpp
@@ -1054,6 +1054,11 @@ void FSlateRHIRenderingPolicy::DrawElements(
 						{
 							PixelShader->SetBlendState(GraphicsPSOInit, Material);
 							FSlateShaderResource* MaskResource = MaterialShaderResource->GetTextureMaskResource();
+							if (MaskResource && (Material->GetBlendMode() == EBlendMode::BLEND_Opaque || Material->GetBlendMode() == EBlendMode::BLEND_Masked))
+							{
+								// Font materials require some form of translucent blending
+								GraphicsPSOInit.BlendState = TStaticBlendState<CW_RGBA, BO_Add, BF_SourceAlpha, BF_InverseSourceAlpha, BO_Add, BF_InverseDestAlpha, BF_One>::GetRHI();
+							}
 
 							GraphicsPSOInit.BoundShaderState.VertexDeclarationRHI = bUseInstancing ? GSlateInstancedVertexDeclaration.VertexDeclarationRHI : GSlateVertexDeclaration.VertexDeclarationRHI;
 							GraphicsPSOInit.BoundShaderState.VertexShaderRHI = VertexShader.GetVertexShader();
