- Hopeful fix for QAGame crash seen one time, with new parallel Slate texture creation

diff --git a/Engine/Source/Runtime/SlateRHIRenderer/Private/SlateRHIResourceManager.cpp b/Engine/Source/Runtime/SlateRHIRenderer/Private/SlateRHIResourceManager.cpp
--- a/Engine/Source/Runtime/SlateRHIRenderer/Private/SlateRHIResourceManager.cpp
+++ b/Engine/Source/Runtime/SlateRHIRenderer/Private/SlateRHIResourceManager.cpp
@@ -330,15 +330,23 @@ void FSlateRHIResourceManager::CreateTextures( const TArray< const FSlateBrush*
 			FNewTextureInfo* Info;
 			{
 				FScopeLock Locker(&Lock);
-				Info = &TextureInfoMap.FindOrAdd(TextureName);
-			}
-	
-			Info->bSrgb = (Brush.ImageType != ESlateBrushImageType::Linear);
 
-			// Only atlas the texture if none of the brushes that use it tile it and the image is srgb
-		
-			Info->bShouldAtlas &= ( Brush.Tiling == ESlateBrushTileType::NoTile && Info->bSrgb && AtlasSize > 0 );
+				bool bSrgb = (Brush.ImageType != ESlateBrushImageType::Linear);
+				bool bShouldAtlas = (Brush.Tiling == ESlateBrushTileType::NoTile && bSrgb && AtlasSize > 0);
 
+				// if it's already in here, then another loop is working on it
+				Info = TextureInfoMap.Find(TextureName);
+				if (Info != nullptr)
+				{
+					// merge in the bShouldAtlas flag to existing, then we are out
+					Info->bShouldAtlas &= bShouldAtlas;
+					return;
+				}
+				Info = &TextureInfoMap.Add(TextureName);
+				Info->bShouldAtlas = bShouldAtlas;
+				Info->bSrgb = bSrgb;
+			}
+		
 			// Texture has been loaded if the texture data is valid
 			if( !Info->TextureData.IsValid() )
 			{
