Added option to use memoryless DepthStencil for Slate UI rendering Slate.MemorylessDepthStencil. Disabled by default

diff --git a/Engine/Source/Runtime/SlateRHIRenderer/Private/SlateRHIRenderer.cpp b/Engine/Source/Runtime/SlateRHIRenderer/Private/SlateRHIRenderer.cpp
--- a/Engine/Source/Runtime/SlateRHIRenderer/Private/SlateRHIRenderer.cpp
+++ b/Engine/Source/Runtime/SlateRHIRenderer/Private/SlateRHIRenderer.cpp
@@ -63,6 +63,12 @@ static TAutoConsoleVariable<float> CVarDrawToVRRenderTarget(
 	TEXT("If enabled while in VR. Slate UI will be drawn into the render target texture where the VR imagery for either eye was rendered, allow the viewer of the HMD to see the UI (for better or worse.)  This render target will then be cropped/scaled into the back buffer, if mirroring is enabled.  When disabled, Slate UI will be drawn on top of the backbuffer (not to the HMD) after the mirror texture has been cropped/scaled into the backbuffer."),
 	ECVF_RenderThreadSafe);
 
+static TAutoConsoleVariable<int32> CVarMemorylessDepthStencil(
+	TEXT("Slate.MemorylessDepthStencil"),
+	0,
+	TEXT("Whether to use memoryless DepthStencil target for Slate. Reduces memory usage and implies that DepthStencil state can't be preserved between Slate renderpasses"),
+	ECVF_RenderThreadSafe);
+
 #if WITH_SLATE_VISUALIZERS
 
 TAutoConsoleVariable<int32> CVarShowSlateOverdraw(
@@ -120,10 +126,12 @@ void FViewportInfo::ConditionallyUpdateDepthBuffer(bool bInRequiresStencilTest,
 {
 	check(IsInRenderingThread());
 
+	bool bWantsMemorylessDepthStencil = (CVarMemorylessDepthStencil.GetValueOnAnyThread() != 0);
+
 	bool bDepthStencilStale =
 		bInRequiresStencilTest &&
 		(!bRequiresStencilTest ||
-		(DepthStencil.IsValid() && (DepthStencil->GetSizeX() != InWidth || DepthStencil->GetSizeY() != InHeight)));
+		(DepthStencil.IsValid() && (DepthStencil->GetSizeX() != InWidth || DepthStencil->GetSizeY() != InHeight || IsMemorylessTexture(DepthStencil) != bWantsMemorylessDepthStencil)));
 
 	bRequiresStencilTest = bInRequiresStencilTest;
 
@@ -142,11 +150,27 @@ void FViewportInfo::RecreateDepthBuffer_RenderThread()
 	{
 		FTexture2DRHIRef ShaderResourceUnused;
 		FRHIResourceCreateInfo CreateInfo(FClearValueBinding::DepthZero);
-		RHICreateTargetableShaderResource2D(Width, Height, PF_DepthStencil, 1, TexCreate_None, TexCreate_DepthStencilTargetable, false, CreateInfo, DepthStencil, ShaderResourceUnused);
+
+		uint32 TargetableTextureFlags = TexCreate_DepthStencilTargetable;
+		if (CVarMemorylessDepthStencil.GetValueOnAnyThread() != 0)
+		{
+			// Use Memoryless target, expecting that DepthStencil content is intermediate and can't be preserved between renderpasses
+			TargetableTextureFlags|= TexCreate_Memoryless;
+		}
+		
+		RHICreateTargetableShaderResource2D(Width, Height, PF_DepthStencil, 1, TexCreate_None, TargetableTextureFlags, false, CreateInfo, DepthStencil, ShaderResourceUnused);
 		check(IsValidRef(DepthStencil));
 	}
 }
 
+bool IsMemorylessTexture(const FTexture2DRHIRef& Tex)
+{
+	if (Tex)
+	{
+		return (Tex->GetFlags() & TexCreate_Memoryless) != 0;
+	}
+	return false;
+}
 
 
 FSlateRHIRenderer::FSlateRHIRenderer(TSharedRef<FSlateFontServices> InSlateFontServices, TSharedRef<FSlateRHIResourceManager> InResourceManager)
@@ -817,7 +841,8 @@ void FSlateRHIRenderer::DrawWindow_RenderThread(FRHICommandListImmediate& RHICmd
 				{
 					check(IsValidRef(ViewportInfo.DepthStencil));
 
-					RPInfo.DepthStencilRenderTarget.Action = MakeDepthStencilTargetActions(ERenderTargetActions::DontLoad_DontStore, ERenderTargetActions::DontLoad_Store);
+					ERenderTargetActions StencilAction = IsMemorylessTexture(ViewportInfo.DepthStencil) ? ERenderTargetActions::DontLoad_DontStore : ERenderTargetActions::DontLoad_Store;
+					RPInfo.DepthStencilRenderTarget.Action = MakeDepthStencilTargetActions(ERenderTargetActions::DontLoad_DontStore, StencilAction);
 					RPInfo.DepthStencilRenderTarget.DepthStencilTarget = ViewportInfo.DepthStencil;
 					RPInfo.DepthStencilRenderTarget.ExclusiveDepthStencil = FExclusiveDepthStencil::DepthNop_StencilWrite;
 				}
@@ -829,7 +854,8 @@ void FSlateRHIRenderer::DrawWindow_RenderThread(FRHICommandListImmediate& RHICmd
 					if (ViewportInfo.bRequiresStencilTest)
 					{
 						// Reset the backbuffer as our color render target and also set a depth stencil buffer
-						RPInfo.DepthStencilRenderTarget.Action = MakeDepthStencilTargetActions(ERenderTargetActions::Load_Store, ERenderTargetActions::Clear_Store);
+						ERenderTargetActions StencilAction = IsMemorylessTexture(ViewportInfo.DepthStencil) ? ERenderTargetActions::Clear_DontStore : ERenderTargetActions::Clear_Store;
+						RPInfo.DepthStencilRenderTarget.Action = MakeDepthStencilTargetActions(ERenderTargetActions::Load_Store, StencilAction);
 						RPInfo.DepthStencilRenderTarget.DepthStencilTarget = ViewportInfo.DepthStencil;
 						RPInfo.DepthStencilRenderTarget.ExclusiveDepthStencil = FExclusiveDepthStencil::DepthWrite_StencilWrite;
 					}
diff --git a/Engine/Source/Runtime/SlateRHIRenderer/Private/SlateRHIRenderer.h b/Engine/Source/Runtime/SlateRHIRenderer/Private/SlateRHIRenderer.h
--- a/Engine/Source/Runtime/SlateRHIRenderer/Private/SlateRHIRenderer.h
+++ b/Engine/Source/Runtime/SlateRHIRenderer/Private/SlateRHIRenderer.h
@@ -148,6 +148,8 @@ struct FViewportInfo : public FRenderResource
 	}	
 };
 
+bool IsMemorylessTexture(const FTexture2DRHIRef& Tex);
+
 struct FFastPathRenderingDataCleanupList
 {
 	TArray<FSlateCachedFastPathRenderingData*, TInlineAllocator<20>> FastPathRenderingDataToRemove;
diff --git a/Engine/Source/Runtime/SlateRHIRenderer/Private/SlateRHIRenderingPolicy.cpp b/Engine/Source/Runtime/SlateRHIRenderer/Private/SlateRHIRenderingPolicy.cpp
--- a/Engine/Source/Runtime/SlateRHIRenderer/Private/SlateRHIRenderingPolicy.cpp
+++ b/Engine/Source/Runtime/SlateRHIRenderer/Private/SlateRHIRenderingPolicy.cpp
@@ -298,9 +298,10 @@ static bool UpdateScissorRect(
 					// #todo-renderpasses this is very gross. If/when this gets refactored we can detect a simple clear or batch up elements by rendertarget (and other stuff)
 					RHICmdList.EndRenderPass();
 					bDidRestartRenderpass = true;
+					ERenderTargetActions StencilAction = IsMemorylessTexture(DepthStencilTarget) ? ERenderTargetActions::DontLoad_DontStore : ERenderTargetActions::Load_Store;
 
 					FRHIRenderPassInfo RPInfo(ColorTarget, ERenderTargetActions::Load_Store);
-					RPInfo.DepthStencilRenderTarget.Action = MakeDepthStencilTargetActions(ERenderTargetActions::DontLoad_DontStore, ERenderTargetActions::Load_Store);
+					RPInfo.DepthStencilRenderTarget.Action = MakeDepthStencilTargetActions(ERenderTargetActions::DontLoad_DontStore, StencilAction);
 					RPInfo.DepthStencilRenderTarget.DepthStencilTarget = DepthStencilTarget;
 					RPInfo.DepthStencilRenderTarget.ExclusiveDepthStencil = FExclusiveDepthStencil::DepthNop_StencilWrite;
 
@@ -397,9 +398,15 @@ static bool UpdateScissorRect(
 
 					// Clear current stencil buffer, we use ELoad/EStore, because we need to keep the stencil around.
 					ERenderTargetLoadAction StencilLoadAction = bClearStencil ? ERenderTargetLoadAction::EClear : ERenderTargetLoadAction::ELoad;
+					ERenderTargetActions StencilAction = MakeRenderTargetActions(StencilLoadAction, ERenderTargetStoreAction::EStore);
+					if (IsMemorylessTexture(DepthStencilTarget))
+					{
+						// We can't preserve content for memoryless targets
+						StencilAction = bClearStencil ? ERenderTargetActions::Clear_DontStore : ERenderTargetActions::DontLoad_DontStore;
+					}
 
 					FRHIRenderPassInfo RPInfo(ColorTarget, ERenderTargetActions::Load_Store);
-					RPInfo.DepthStencilRenderTarget.Action = MakeDepthStencilTargetActions(ERenderTargetActions::DontLoad_DontStore, MakeRenderTargetActions(StencilLoadAction, ERenderTargetStoreAction::EStore));
+					RPInfo.DepthStencilRenderTarget.Action = MakeDepthStencilTargetActions(ERenderTargetActions::DontLoad_DontStore, StencilAction);
 					RPInfo.DepthStencilRenderTarget.DepthStencilTarget = DepthStencilTarget;
 					RPInfo.DepthStencilRenderTarget.ExclusiveDepthStencil = FExclusiveDepthStencil::DepthNop_StencilWrite;
 					RHICmdList.BeginRenderPass(RPInfo, TEXT("SlateUpdateScissorRect_ClearStencil"));
@@ -699,7 +706,7 @@ void FSlateRHIRenderingPolicy::DrawElements(
 			RPInfo.DepthStencilRenderTarget.DepthStencilTarget = DepthStencilTarget;
 			if (DepthStencilTarget)
 			{
-				RPInfo.DepthStencilRenderTarget.Action = EDepthStencilTargetActions::LoadDepthStencil_StoreDepthStencil;
+				RPInfo.DepthStencilRenderTarget.Action = IsMemorylessTexture(DepthStencilTarget) ? EDepthStencilTargetActions::DontLoad_DontStore : EDepthStencilTargetActions::LoadDepthStencil_StoreDepthStencil;
 				RPInfo.DepthStencilRenderTarget.ExclusiveDepthStencil = FExclusiveDepthStencil::DepthWrite_StencilWrite;
 			}
 			else
