Fix for Background Blur in UMG is filled with black when UI.CompositeMode is used in HDR

diff --git a/Engine/Source/Runtime/SlateRHIRenderer/Private/Slate3DRenderer.cpp b/Engine/Source/Runtime/SlateRHIRenderer/Private/Slate3DRenderer.cpp
--- a/Engine/Source/Runtime/SlateRHIRenderer/Private/Slate3DRenderer.cpp
+++ b/Engine/Source/Runtime/SlateRHIRenderer/Private/Slate3DRenderer.cpp
@@ -190,6 +190,7 @@ void FSlate3DRenderer::DrawWindowToTarget_RenderThread(FRHICommandListImmediate&
 					InRHICmdList,
 					BackBufferTarget,
 					ColorTarget,
+					ColorTarget,
 					DepthStencil,
 					BatchData.GetFirstRenderBatchIndex(),
 					BatchData.GetRenderBatches(),
diff --git a/Engine/Source/Runtime/SlateRHIRenderer/Private/SlateRHIRenderer.cpp b/Engine/Source/Runtime/SlateRHIRenderer/Private/SlateRHIRenderer.cpp
--- a/Engine/Source/Runtime/SlateRHIRenderer/Private/SlateRHIRenderer.cpp
+++ b/Engine/Source/Runtime/SlateRHIRenderer/Private/SlateRHIRenderer.cpp
@@ -747,6 +747,7 @@ void FSlateRHIRenderer::DrawWindow_RenderThread(FRHICommandListImmediate& RHICmd
 
 			FTexture2DRHIRef ViewportRT = bRenderedStereo ? nullptr : ViewportInfo.GetRenderTargetTexture();
 			FTexture2DRHIRef BackBuffer = (ViewportRT) ? ViewportRT : RHICmdList.GetViewportBackBuffer(ViewportInfo.ViewportRHI);
+			FTexture2DRHIRef PostProcessBuffer = BackBuffer;	// If compositing UI then this will be different to the back buffer
 
 			const uint32 ViewportWidth = (ViewportRT) ? ViewportRT->GetSizeX() : ViewportInfo.Width;
 			const uint32 ViewportHeight = (ViewportRT) ? ViewportRT->GetSizeY() : ViewportInfo.Height;
@@ -880,6 +881,7 @@ void FSlateRHIRenderer::DrawWindow_RenderThread(FRHICommandListImmediate& RHICmd
 							RHICmdList,
 							BackBufferTarget,
 							BackBuffer,
+							PostProcessBuffer,
 							ViewportInfo.bRequiresStencilTest ? ViewportInfo.DepthStencil : EmptyTarget,
 							BatchData.GetFirstRenderBatchIndex(),
 							BatchData.GetRenderBatches(),
diff --git a/Engine/Source/Runtime/SlateRHIRenderer/Private/SlateRHIRenderingPolicy.cpp b/Engine/Source/Runtime/SlateRHIRenderer/Private/SlateRHIRenderingPolicy.cpp
--- a/Engine/Source/Runtime/SlateRHIRenderer/Private/SlateRHIRenderingPolicy.cpp
+++ b/Engine/Source/Runtime/SlateRHIRenderer/Private/SlateRHIRenderingPolicy.cpp
@@ -562,6 +562,7 @@ void FSlateRHIRenderingPolicy::DrawElements(
 	FRHICommandListImmediate& RHICmdList,
 	FSlateBackBuffer& BackBuffer,
 	FTexture2DRHIRef& ColorTarget,
+	FTexture2DRHIRef& PostProcessTexture,
 	FTexture2DRHIRef& DepthStencilTarget,
 	int32 FirstBatchIndex,
 	const TArray<FSlateRenderBatch>& RenderBatches,
@@ -1133,10 +1134,10 @@ void FSlateRHIRenderingPolicy::DrawElements(
 				const FVector4& QuadPositionData = ShaderParams.PixelParams;
 
 				FPostProcessRectParams RectParams;
-				RectParams.SourceTexture = BackBuffer.GetRenderTargetTexture();
-				RectParams.SourceRect = FSlateRect(0, 0, BackBuffer.GetSizeXY().X, BackBuffer.GetSizeXY().Y);
+				RectParams.SourceTexture = PostProcessTexture->GetTexture2D();
+				RectParams.SourceRect = FSlateRect(0, 0, PostProcessTexture->GetSizeX(), PostProcessTexture->GetSizeY());
 				RectParams.DestRect = FSlateRect(QuadPositionData.X, QuadPositionData.Y, QuadPositionData.Z, QuadPositionData.W);
-				RectParams.SourceTextureSize = BackBuffer.GetSizeXY();
+				RectParams.SourceTextureSize = PostProcessTexture->GetSizeXY();
 
 				RectParams.RestoreStateFunc = [&](FRHICommandListImmediate&InRHICmdList, FGraphicsPipelineStateInitializer& InGraphicsPSOInit) {
 					return UpdateScissorRect(
diff --git a/Engine/Source/Runtime/SlateRHIRenderer/Private/SlateRHIRenderingPolicy.h b/Engine/Source/Runtime/SlateRHIRenderer/Private/SlateRHIRenderingPolicy.h
--- a/Engine/Source/Runtime/SlateRHIRenderer/Private/SlateRHIRenderingPolicy.h
+++ b/Engine/Source/Runtime/SlateRHIRenderer/Private/SlateRHIRenderingPolicy.h
@@ -55,7 +55,8 @@ public:
 
 	void BuildRenderingBuffers(FRHICommandListImmediate& RHICmdList, FSlateBatchData& InBatchData);
 
-	void DrawElements(FRHICommandListImmediate& RHICmdList, class FSlateBackBuffer& BackBuffer, FTexture2DRHIRef& ColorTarget, FTexture2DRHIRef& DepthStencilTarget, int32 FirstBatchIndex, const TArray<FSlateRenderBatch>& RenderBatches, const FSlateRenderingParams& Params);
+	void DrawElements(FRHICommandListImmediate& RHICmdList, class FSlateBackBuffer& BackBuffer, FTexture2DRHIRef& ColorTarget, FTexture2DRHIRef& PostProcessBuffer,
+		FTexture2DRHIRef& DepthStencilTarget, int32 FirstBatchIndex,  const TArray<FSlateRenderBatch>& RenderBatches, const FSlateRenderingParams& Params);
 
 	virtual TSharedRef<FSlateShaderResourceManager> GetResourceManager() const override { return ResourceManager; }
 	virtual bool IsVertexColorInLinearSpace() const override { return false; }
