Clamp bounding-box of stencil quad to viewport extent in "UpdateScissorRect".

#rnx

diff --git a/Engine/Source/Runtime/SlateRHIRenderer/Private/SlateRHIRenderingPolicy.cpp b/Engine/Source/Runtime/SlateRHIRenderer/Private/SlateRHIRenderingPolicy.cpp
--- a/Engine/Source/Runtime/SlateRHIRenderer/Private/SlateRHIRenderingPolicy.cpp
+++ b/Engine/Source/Runtime/SlateRHIRenderer/Private/SlateRHIRenderingPolicy.cpp
@@ -371,21 +371,27 @@ static bool UpdateScissorRect(
 					const FSlateClippingZone& MaskQuad = StencilQuads.Last();
 					const FSlateRect LastStencilBoundingBox = MaskQuad.GetBoundingBox().Round();
 
-					const FIntPoint SizeXY = BackBuffer.GetSizeXY();
-					const FVector2D ViewSize((float)SizeXY.X, (float)SizeXY.Y);
+					FSlateRect ScissorRect = LastStencilBoundingBox.OffsetBy(ViewTranslation2D);
 
-					const FVector2D TopLeft = FMath::Min(FMath::Max(LastStencilBoundingBox.GetTopLeft() + ViewTranslation2D, FVector2D(0.0f, 0.0f)), ViewSize);
-					const FVector2D BottomRight = FMath::Min(FMath::Max(LastStencilBoundingBox.GetBottomRight() + ViewTranslation2D, FVector2D(0.0f, 0.0f)), ViewSize);
+					// Chosen stencil quad might have some coordinates outside the viewport.
+					// After turning it into a bounding box, this box must be clamped to the current viewport,
+					// as scissors outside the viewport don't make sense (and cause assertions to fail).
+					const FIntPoint BackBufferSize = BackBuffer.GetSizeXY();
+					ScissorRect.Left = FMath::Clamp(ScissorRect.Left, 0.0f, static_cast<float>(BackBufferSize.X));
+					ScissorRect.Top = FMath::Clamp(ScissorRect.Top, 0.0f, static_cast<float>(BackBufferSize.Y));
+					ScissorRect.Right = FMath::Clamp(ScissorRect.Right, ScissorRect.Left, static_cast<float>(BackBufferSize.X));
+					ScissorRect.Bottom = FMath::Clamp(ScissorRect.Bottom, ScissorRect.Top, static_cast<float>(BackBufferSize.Y));
 
 					if (bSwitchVerticalAxis)
 					{
-						const int32 MinY = (ViewSize.Y - BottomRight.Y);
-						const int32 MaxY = (ViewSize.Y - TopLeft.Y);
-						RHICmdList.SetScissorRect(true, TopLeft.X, MinY, BottomRight.X, MaxY);
+						const FIntPoint ViewSize = BackBuffer.GetSizeXY();
+						const int32 MinY = (ViewSize.Y - ScissorRect.Bottom);
+						const int32 MaxY = (ViewSize.Y - ScissorRect.Top);
+						RHICmdList.SetScissorRect(true, ScissorRect.Left, MinY, ScissorRect.Right, MaxY);
 					}
 					else
 					{
-						RHICmdList.SetScissorRect(true, TopLeft.X, TopLeft.Y, BottomRight.X, BottomRight.Y);
+						RHICmdList.SetScissorRect(true, ScissorRect.Left, ScissorRect.Top, ScissorRect.Right, ScissorRect.Bottom);
 					}
 				}
 
