Slate - Adding more debugging and an alternative path for getting the material proxy on the render thread.  In rare cases seeing materials being GCed, adding additiona debugging to try and trace that down.

diff --git a/Engine/Source/Runtime/SlateRHIRenderer/Private/SlateMaterialResource.cpp b/Engine/Source/Runtime/SlateRHIRenderer/Private/SlateMaterialResource.cpp
--- a/Engine/Source/Runtime/SlateRHIRenderer/Private/SlateMaterialResource.cpp
+++ b/Engine/Source/Runtime/SlateRHIRenderer/Private/SlateMaterialResource.cpp
@@ -5,16 +5,32 @@
 #include "Styling/SlateBrush.h"
 
 
-FSlateMaterialResource::FSlateMaterialResource(const UMaterialInterface& InMaterial, const FVector2D& InImageSize, FSlateShaderResource* InTextureMask )
-	: MaterialObject( &InMaterial )
+FSlateMaterialResource::FSlateMaterialResource(const UMaterialInterface& InMaterialResource, const FVector2D& InImageSize, FSlateShaderResource* InTextureMask )
+	: MaterialObject( &InMaterialResource)
 	, SlateProxy( new FSlateShaderResourceProxy )
 	, TextureMaskResource( InTextureMask )
 	, Width(FMath::RoundToInt(InImageSize.X))
 	, Height(FMath::RoundToInt(InImageSize.Y))
 {
+#if SLATE_CHECK_UOBJECT_RENDER_RESOURCES
+	ensure(!InMaterialResource.IsPendingKill());
+#endif
+
 	SlateProxy->ActualSize = InImageSize.IntPoint();
 	SlateProxy->Resource = this;
 
+	MaterialProxy = InMaterialResource.GetRenderProxy();
+
+#if SLATE_CHECK_UOBJECT_RENDER_RESOURCES
+	ensure(!MaterialProxy->IsDeleted());
+	ensure(!MaterialProxy->IsMarkedForGarbageCollection());
+#endif
+
+	if (MaterialProxy->IsDeleted() || MaterialProxy->IsMarkedForGarbageCollection())
+	{
+		MaterialProxy = nullptr;
+	}
+
 #if SLATE_CHECK_UOBJECT_RENDER_RESOURCES
 	MaterialObjectWeakPtr = MaterialObject; 
 	UpdateMaterialName();
@@ -31,7 +47,22 @@ FSlateMaterialResource::~FSlateMaterialResource()
 
 void FSlateMaterialResource::UpdateMaterial(const UMaterialInterface& InMaterialResource, const FVector2D& InImageSize, FSlateShaderResource* InTextureMask )
 {
+#if SLATE_CHECK_UOBJECT_RENDER_RESOURCES
+	ensure(!InMaterialResource.IsPendingKill());
+#endif
+
 	MaterialObject = &InMaterialResource;
+	MaterialProxy = InMaterialResource.GetRenderProxy();
+
+#if SLATE_CHECK_UOBJECT_RENDER_RESOURCES
+	ensure(!MaterialProxy->IsDeleted());
+	ensure(!MaterialProxy->IsMarkedForGarbageCollection());
+#endif
+
+	if (MaterialProxy->IsDeleted() || MaterialProxy->IsMarkedForGarbageCollection())
+	{
+		MaterialProxy = nullptr;
+	}
 
 #if SLATE_CHECK_UOBJECT_RENDER_RESOURCES
 	MaterialObjectWeakPtr = MaterialObject;
diff --git a/Engine/Source/Runtime/SlateRHIRenderer/Private/SlateMaterialResource.h b/Engine/Source/Runtime/SlateRHIRenderer/Private/SlateMaterialResource.h
--- a/Engine/Source/Runtime/SlateRHIRenderer/Private/SlateMaterialResource.h
+++ b/Engine/Source/Runtime/SlateRHIRenderer/Private/SlateMaterialResource.h
@@ -25,10 +25,13 @@ public:
 	void ResetMaterial();
 
 	/** @return The material render proxy */
-	FMaterialRenderProxy* GetRenderProxy() const { return MaterialObject ? MaterialObject->GetRenderProxy() : nullptr; }
+	FMaterialRenderProxy* GetRenderProxy() const { return MaterialProxy; }
 
 	/** @return the material object */
-	const UMaterialInterface* GetMaterialObject() const { return MaterialObject; }
+	const UMaterialInterface* GetMaterialObject() const
+	{
+		return MaterialObject;
+	}
 
 	/** Slate proxy used for batching the material */
 	FSlateShaderResourceProxy* GetResourceProxy() const { return SlateProxy; }
@@ -41,6 +44,7 @@ public:
 
 private:
 	const class UMaterialInterface* MaterialObject;
+	class FMaterialRenderProxy* MaterialProxy;
 
 #if SLATE_CHECK_UOBJECT_RENDER_RESOURCES
 	// Used to guard against crashes when the material object is deleted.  This is expensive so we do not do it in shipping
diff --git a/Engine/Source/Runtime/SlateRHIRenderer/Private/SlateRHIRenderingPolicy.cpp b/Engine/Source/Runtime/SlateRHIRenderer/Private/SlateRHIRenderingPolicy.cpp
--- a/Engine/Source/Runtime/SlateRHIRenderer/Private/SlateRHIRenderingPolicy.cpp
+++ b/Engine/Source/Runtime/SlateRHIRenderer/Private/SlateRHIRenderingPolicy.cpp
@@ -1026,11 +1026,11 @@ void FSlateRHIRenderingPolicy::DrawElements(
 				const FSceneView& ActiveSceneView = *SceneViews[ActiveSceneIndex];
 
 				FSlateMaterialResource* MaterialShaderResource = (FSlateMaterialResource*)ShaderResource;
-				if (MaterialShaderResource->GetMaterialObject() != nullptr)
+				if (FMaterialRenderProxy* MaterialRenderProxy = MaterialShaderResource->GetRenderProxy())
 				{
 					MaterialShaderResource->CheckForStaleResources();
 
-					FMaterialRenderProxy* MaterialRenderProxy = MaterialShaderResource->GetRenderProxy();
+					//FMaterialRenderProxy* MaterialRenderProxy = MaterialShaderResource->GetRenderProxy();
 
 					const FMaterial* Material = MaterialRenderProxy->GetMaterial(ActiveSceneView.GetFeatureLevel());
 
