Optimize Slate texture loading during startup
  - 450ms to 111ms in warm
  - 2.5s to 378ms in cold (No AV exclusions)
  - 1.4s to 128ms in cold (AV exclusions on png files)

  - Cold cache = Cleared standby pages, mostly equivalent to a reboot
  - Warm cache = 2nd run of 2 consecutive runs

#rnx

diff --git a/Engine/Source/Runtime/SlateCore/Public/Textures/SlateTextureData.h b/Engine/Source/Runtime/SlateCore/Public/Textures/SlateTextureData.h
--- a/Engine/Source/Runtime/SlateCore/Public/Textures/SlateTextureData.h
+++ b/Engine/Source/Runtime/SlateCore/Public/Textures/SlateTextureData.h
@@ -22,6 +22,15 @@ struct SLATECORE_API FSlateTextureData
 		INC_MEMORY_STAT_BY( STAT_SlateTextureDataMemory, Bytes.GetAllocatedSize() );
 	}
 
+	FSlateTextureData(uint32 InWidth, uint32 InHeight, uint32 InBytesPerPixel, TArray<uint8>&& InBytes)
+		: Bytes(InBytes)
+		, Width(InWidth)
+		, Height(InHeight)
+		, BytesPerPixel(InBytesPerPixel)
+	{
+		INC_MEMORY_STAT_BY(STAT_SlateTextureDataMemory, Bytes.GetAllocatedSize());
+	}
+
 	/**
 	 * Constructor to create texture data by copying from a pointer instead of an array
 	 * @param InBuffer Pointer to Texture data (must contain InWidth*InHeight*InBytesPerPixel bytes).
diff --git a/Engine/Source/Runtime/SlateRHIRenderer/Private/SlateRHIResourceManager.cpp b/Engine/Source/Runtime/SlateRHIRenderer/Private/SlateRHIResourceManager.cpp
--- a/Engine/Source/Runtime/SlateRHIRenderer/Private/SlateRHIResourceManager.cpp
+++ b/Engine/Source/Runtime/SlateRHIRenderer/Private/SlateRHIResourceManager.cpp
@@ -311,78 +311,80 @@ void FSlateRHIResourceManager::Tick(float DeltaSeconds)
 
 void FSlateRHIResourceManager::CreateTextures( const TArray< const FSlateBrush* >& Resources )
 {
-	FModuleManager::LoadModuleChecked<IImageWrapperModule>(FName("ImageWrapper"));
-
-	TMap<FName,FNewTextureInfo> TextureInfoMap;
-	FCriticalSection Lock;
+	TRACE_CPUPROFILER_EVENT_SCOPE(FSlateRHIResourceManager::CreateTextures)
 
-	// reserve space so we don't reallocate memory during the ParallelFor
-	TextureInfoMap.Reserve(Resources.Num());
+	struct FLoadTextureInfo : public FNewTextureInfo
+	{
+		bool bSucceeded = false;
+		const FSlateBrush* Brush;
+	};
 
-	ParallelFor(Resources.Num(), [this, &Resources, &Lock, &TextureInfoMap](int32 ResourceIndex)
+	TMap<FName, FLoadTextureInfo> TextureInfoMap;
+	for( int32 ResourceIndex = 0; ResourceIndex < Resources.Num(); ++ResourceIndex )
 	{
-		const uint32 Stride = GPixelFormats[PF_R8G8B8A8].BlockBytes;
 		const FSlateBrush& Brush = *Resources[ResourceIndex];
 		const FName TextureName = Brush.GetResourceName();
 		if( TextureName != NAME_None && !Brush.HasUObject() && !Brush.IsDynamicallyLoaded() && !ResourceMap.Contains(TextureName) )
 		{
-			// Find the texture or add it if it doesnt exist (only load the texture once)
-			FNewTextureInfo* Info;
-			{
-				FScopeLock Locker(&Lock);
+			// Find the texture or add it if it doesn't exist (only load the texture once)
+			FLoadTextureInfo& Info = TextureInfoMap.FindOrAdd( TextureName );
 
-				bool bSrgb = (Brush.ImageType != ESlateBrushImageType::Linear);
-				bool bShouldAtlas = (Brush.Tiling == ESlateBrushTileType::NoTile && bSrgb && AtlasSize > 0);
+			Info.Brush = &Brush;
+			Info.bSrgb = (Brush.ImageType != ESlateBrushImageType::Linear);
 
-				// if it's already in here, then another loop is working on it
-				Info = TextureInfoMap.Find(TextureName);
-				if (Info != nullptr)
-				{
-					// merge in the bShouldAtlas flag to existing, then we are out
-					Info->bShouldAtlas &= bShouldAtlas;
-					return;
-				}
-				Info = &TextureInfoMap.Add(TextureName);
-				Info->bShouldAtlas = bShouldAtlas;
-				Info->bSrgb = bSrgb;
-			}
+			// Only atlas the texture if none of the brushes that use it tile it and the image is srgb
 		
-			// Texture has been loaded if the texture data is valid
-			if( !Info->TextureData.IsValid() )
-			{
-				uint32 Width = 0;
-				uint32 Height = 0;
-				TArray<uint8> RawData;
-				const bool bSucceeded = LoadTexture( Brush, Width, Height, RawData );
-				if (bSucceeded)
-				{
-					Info->TextureData = MakeShareable(new FSlateTextureData(Width, Height, Stride, RawData));
+			Info.bShouldAtlas &= ( Brush.Tiling == ESlateBrushTileType::NoTile && Info.bSrgb && AtlasSize > 0 );
+		}
+	}
 
-					const bool bTooLargeForAtlas = (Width >= (uint32)MaxAltasedTextureSize.X || Height >= (uint32)MaxAltasedTextureSize.Y || Width >= AtlasSize || Height >= AtlasSize);
+	// This must be loaded before going multi-threaded
+	FModuleManager::LoadModuleChecked<IImageWrapperModule>(FName("ImageWrapper"));
 
-					Info->bShouldAtlas &= !bTooLargeForAtlas;
-				}
+	TArray<FName> TexturesToLoad;
+	TextureInfoMap.GenerateKeyArray(TexturesToLoad);
 
-				if( !bSucceeded || !ensureMsgf( Info->TextureData->GetRawBytes().Num() > 0, TEXT("Slate resource: (%s) contains no data"), *TextureName.ToString() ) )
-				{
-					FScopeLock Locker(&Lock);
-					TextureInfoMap.Remove( TextureName );
-				}
-			}
+	const uint32 Stride = GPixelFormats[PF_R8G8B8A8].BlockBytes;
+	ParallelFor(
+		TexturesToLoad.Num(),
+		[&](int32 TextureIndex)
+		{
+			FName& TextureName = TexturesToLoad[TextureIndex];
+			FLoadTextureInfo& Info = TextureInfoMap.FindChecked(TextureName);
+
+			uint32 Width = 0;
+			uint32 Height = 0;
+			TArray<uint8> RawData;
+			Info.bSucceeded = LoadTexture(*Info.Brush, Width, Height, RawData);
+
+			Info.TextureData = MakeShareable(new FSlateTextureData(Width, Height, Stride, MoveTemp(RawData)));
+
+			const bool bTooLargeForAtlas = (Width >= (uint32)MaxAltasedTextureSize.X || Height >= (uint32)MaxAltasedTextureSize.Y || Width >= AtlasSize || Height >= AtlasSize);
+			Info.bShouldAtlas &= !bTooLargeForAtlas;
 		}
-	});
+	);
+
+	// Remove any texture that didn't succeed
+	for (FName& TextureName : TexturesToLoad)
+	{
+		FLoadTextureInfo& Info = TextureInfoMap.FindChecked(TextureName);
+		if (!Info.bSucceeded || !ensureMsgf(Info.TextureData->GetRawBytes().Num() > 0, TEXT("Slate resource: (%s) contains no data"), *TextureName.ToString()))
+		{
+			TextureInfoMap.Remove(TextureName);
+		}
+	}
 
 	// Sort textures by size.  The largest textures are atlased first which creates a more compact atlas
 	TextureInfoMap.ValueSort( FCompareFNewTextureInfoByTextureSize() );
 
-	for( TMap<FName,FNewTextureInfo>::TConstIterator It(TextureInfoMap); It; ++It )
+	for( TMap<FName,FLoadTextureInfo>::TConstIterator It(TextureInfoMap); It; ++It )
 	{
-		const FNewTextureInfo& Info = It.Value();
+		const FLoadTextureInfo& Info = It.Value();
 		FName TextureName = It.Key();
 
 		checkSlow( TextureName != NAME_None );
 
-		FSlateShaderResourceProxy* NewTexture = GenerateTextureResource( Info, TextureName );
+		FSlateShaderResourceProxy* NewTexture = GenerateTextureResource( Info );
 
 		ResourceMap.Add( TextureName, NewTexture );
 	}
@@ -402,9 +404,6 @@ bool FSlateRHIResourceManager::LoadTexture( const FSlateBrush& InBrush, uint32&
  */
 bool FSlateRHIResourceManager::LoadTexture( const FName& TextureName, const FString& ResourcePath, uint32& Width, uint32& Height, TArray<uint8>& DecodedImage )
 {
-	// @todo loadtime: nothing in this function uses SlateRendering, so why is this here? I think it's a bad check, and breaks ParallelFor above
-//	checkSlow( IsThreadSafeForSlateRendering() );
-
 	bool bSucceeded = false;
 	uint32 BytesPerPixel = 4;
 
