Fix build error in SlateRHIResourceManager due to bad merge

#rnx

diff --git a/Engine/Source/Runtime/SlateRHIRenderer/Private/SlateRHIResourceManager.cpp b/Engine/Source/Runtime/SlateRHIRenderer/Private/SlateRHIResourceManager.cpp
--- a/Engine/Source/Runtime/SlateRHIRenderer/Private/SlateRHIResourceManager.cpp
+++ b/Engine/Source/Runtime/SlateRHIRenderer/Private/SlateRHIResourceManager.cpp
@@ -384,7 +384,7 @@ void FSlateRHIResourceManager::CreateTextures( const TArray< const FSlateBrush*
 
 		checkSlow( TextureName != NAME_None );
 
-		FSlateShaderResourceProxy* NewTexture = GenerateTextureResource( Info );
+		FSlateShaderResourceProxy* NewTexture = GenerateTextureResource( Info, TextureName );
 
 		ResourceMap.Add( TextureName, NewTexture );
 	}
