Slate - Additional validation about every material we paint with to attempt to prevent GCed materials or materials we can't stop from being collected being drawn with.

diff --git a/Engine/Source/Runtime/SlateCore/Private/Rendering/DrawElements.cpp b/Engine/Source/Runtime/SlateCore/Private/Rendering/DrawElements.cpp
--- a/Engine/Source/Runtime/SlateCore/Private/Rendering/DrawElements.cpp
+++ b/Engine/Source/Runtime/SlateCore/Private/Rendering/DrawElements.cpp
@@ -356,6 +356,18 @@ void FSlateDrawElement::MakeShapedText(FSlateWindowElementList& ElementList, uin
 		return;
 	}
 
+#if SLATE_CHECK_UOBJECT_RENDER_RESOURCES
+	if (const UObject* FontMaterial = InShapedGlyphSequence->GetFontMaterial())
+	{
+		ensureMsgf(!FontMaterial->IsPendingKill(), TEXT("Font Material %s Is Marked Pending Kill"), *FontMaterial->GetName());
+	}
+
+	if (const UObject* OutlineMaterial = InShapedGlyphSequence->GetFontOutlineSettings().OutlineMaterial)
+	{
+		ensureMsgf(!OutlineMaterial->IsPendingKill(), TEXT("Outline Material %s Is Marked Pending Kill"), *OutlineMaterial->GetName());
+	}
+#endif
+
 	FSlateDrawElement& Element = ElementList.AddUninitialized();
 
 	FSlateShapedTextPayload& DataPayload = ElementList.CreatePayload<FSlateShapedTextPayload>(Element);
diff --git a/Engine/Source/Runtime/SlateRHIRenderer/Private/SlateMaterialResource.cpp b/Engine/Source/Runtime/SlateRHIRenderer/Private/SlateMaterialResource.cpp
--- a/Engine/Source/Runtime/SlateRHIRenderer/Private/SlateMaterialResource.cpp
+++ b/Engine/Source/Runtime/SlateRHIRenderer/Private/SlateMaterialResource.cpp
@@ -12,64 +12,64 @@ FSlateMaterialResource::FSlateMaterialResource(const UMaterialInterface& InMater
 	, Width(FMath::RoundToInt(InImageSize.X))
 	, Height(FMath::RoundToInt(InImageSize.Y))
 {
+	MaterialProxy = InMaterialResource.GetRenderProxy();
+
 #if SLATE_CHECK_UOBJECT_RENDER_RESOURCES
-	ensure(!InMaterialResource.IsPendingKill());
+	MaterialObjectWeakPtr = MaterialObject;
+	UpdateMaterialName();
+
+	ensureMsgf(!InMaterialResource.IsPendingKill(), TEXT("Material %s IsPendingKill"), *DebugName.ToString());
+
+	ensureMsgf(MaterialProxy, TEXT("Material %s Render Proxy Is Missing"), *DebugName.ToString());
+	if (MaterialProxy)
+	{
+		ensureMsgf(!MaterialProxy->IsDeleted(), TEXT("Material %s Render Proxy Is Deleted"), *DebugName.ToString());
+		ensureMsgf(!MaterialProxy->IsMarkedForGarbageCollection(), TEXT("Material %s Render Proxy Is Marked For GC"), *DebugName.ToString());
+	}
 #endif
 
 	SlateProxy->ActualSize = InImageSize.IntPoint();
 	SlateProxy->Resource = this;
 
-	MaterialProxy = InMaterialResource.GetRenderProxy();
-
-#if SLATE_CHECK_UOBJECT_RENDER_RESOURCES
-	ensure(!MaterialProxy->IsDeleted());
-	ensure(!MaterialProxy->IsMarkedForGarbageCollection());
-#endif
-
-	if (MaterialProxy->IsDeleted() || MaterialProxy->IsMarkedForGarbageCollection())
+	if (MaterialProxy && (MaterialProxy->IsDeleted() || MaterialProxy->IsMarkedForGarbageCollection()))
 	{
 		MaterialProxy = nullptr;
 	}
-
-#if SLATE_CHECK_UOBJECT_RENDER_RESOURCES
-	MaterialObjectWeakPtr = MaterialObject; 
-	UpdateMaterialName();
-#endif
 }
 
 FSlateMaterialResource::~FSlateMaterialResource()
 {
-	if(SlateProxy)
+	if (SlateProxy)
 	{
 		delete SlateProxy;
 	}
 }
 
-void FSlateMaterialResource::UpdateMaterial(const UMaterialInterface& InMaterialResource, const FVector2D& InImageSize, FSlateShaderResource* InTextureMask )
+void FSlateMaterialResource::UpdateMaterial(const UMaterialInterface& InMaterialResource, const FVector2D& InImageSize, FSlateShaderResource* InTextureMask)
 {
-#if SLATE_CHECK_UOBJECT_RENDER_RESOURCES
-	ensure(!InMaterialResource.IsPendingKill());
-#endif
-
 	MaterialObject = &InMaterialResource;
 	MaterialProxy = InMaterialResource.GetRenderProxy();
 
 #if SLATE_CHECK_UOBJECT_RENDER_RESOURCES
-	ensure(!MaterialProxy->IsDeleted());
-	ensure(!MaterialProxy->IsMarkedForGarbageCollection());
+	MaterialObjectWeakPtr = MaterialObject;
+	UpdateMaterialName();
+
+	ensureMsgf(!InMaterialResource.IsPendingKill(), TEXT("Material %s IsPendingKill"), *DebugName.ToString());
+
+	ensureMsgf(MaterialProxy, TEXT("Material %s Render Proxy Is Missing"), *DebugName.ToString());
+	if (MaterialProxy)
+	{
+		ensureMsgf(!MaterialProxy->IsDeleted(), TEXT("Material %s Render Proxy Is Deleted"), *DebugName.ToString());
+		ensureMsgf(!MaterialProxy->IsMarkedForGarbageCollection(), TEXT("Material %s Render Proxy Is Marked For GC"), *DebugName.ToString());
+	}
 #endif
 
-	if (MaterialProxy->IsDeleted() || MaterialProxy->IsMarkedForGarbageCollection())
+	if (MaterialProxy && (MaterialProxy->IsDeleted() || MaterialProxy->IsMarkedForGarbageCollection()))
 	{
 		MaterialProxy = nullptr;
 	}
 
-#if SLATE_CHECK_UOBJECT_RENDER_RESOURCES
-	MaterialObjectWeakPtr = MaterialObject;
-	UpdateMaterialName();
-#endif
-
-	if( !SlateProxy )
+	if (!SlateProxy)
 	{
 		SlateProxy = new FSlateShaderResourceProxy;
 	}
