Deactivate multi threaded ParallelFor while investigating a deadlock in Shipping Client

#rnx

diff --git a/Engine/Source/Runtime/SlateRHIRenderer/Private/SlateRHIResourceManager.cpp b/Engine/Source/Runtime/SlateRHIRenderer/Private/SlateRHIResourceManager.cpp
--- a/Engine/Source/Runtime/SlateRHIRenderer/Private/SlateRHIResourceManager.cpp
+++ b/Engine/Source/Runtime/SlateRHIRenderer/Private/SlateRHIResourceManager.cpp
@@ -361,7 +361,9 @@ void FSlateRHIResourceManager::CreateTextures( const TArray< const FSlateBrush*
 
 			const bool bTooLargeForAtlas = (Width >= (uint32)MaxAltasedTextureSize.X || Height >= (uint32)MaxAltasedTextureSize.Y || Width >= AtlasSize || Height >= AtlasSize);
 			Info.bShouldAtlas &= !bTooLargeForAtlas;
-		}
+		},
+		// Deactivate multi threaded ParallelFor while investigating a deadlock in Shipping Client
+		EParallelForFlags::ForceSingleThread
 	);
 
 	// Remove any texture that didn't succeed
