Reenable parallel texture loading in Slate as the Pak loading deadlock is now fixed.
  - Fixed since integration CL 12455482, original fix in CL 11724753.

#rnx

diff --git a/Engine/Source/Runtime/SlateRHIRenderer/Private/SlateRHIResourceManager.cpp b/Engine/Source/Runtime/SlateRHIRenderer/Private/SlateRHIResourceManager.cpp
--- a/Engine/Source/Runtime/SlateRHIRenderer/Private/SlateRHIResourceManager.cpp
+++ b/Engine/Source/Runtime/SlateRHIRenderer/Private/SlateRHIResourceManager.cpp
@@ -362,8 +362,7 @@ void FSlateRHIResourceManager::CreateTextures( const TArray< const FSlateBrush*
 			const bool bTooLargeForAtlas = (Width >= (uint32)MaxAltasedTextureSize.X || Height >= (uint32)MaxAltasedTextureSize.Y || Width >= AtlasSize || Height >= AtlasSize);
 			Info.bShouldAtlas &= !bTooLargeForAtlas;
 		},
-		// Deactivate multi threaded ParallelFor while investigating a deadlock in Shipping Client
-		EParallelForFlags::ForceSingleThread
+		EParallelForFlags::Unbalanced
 	);
 
 	// Remove any texture that didn't succeed
