#UE4 IWYU fix

diff --git a/Engine/Source/Runtime/SlateRHIRenderer/Private/SlateRHIResourceManager.cpp b/Engine/Source/Runtime/SlateRHIRenderer/Private/SlateRHIResourceManager.cpp
--- a/Engine/Source/Runtime/SlateRHIRenderer/Private/SlateRHIResourceManager.cpp
+++ b/Engine/Source/Runtime/SlateRHIRenderer/Private/SlateRHIResourceManager.cpp
@@ -6,6 +6,7 @@
 #include "Misc/FileHelper.h"
 #include "Misc/ConfigCacheIni.h"
 #include "Misc/CoreDelegates.h"
+#include "Misc/ParallelFor.h"
 #include "Modules/ModuleManager.h"
 #include "Styling/SlateStyleRegistry.h"
 #include "Styling/ISlateStyle.h"
