Fix memleak in the slate 3d renderer by never cleaning up memstack allocator allocations

diff --git a/Engine/Source/Runtime/SlateRHIRenderer/Private/Slate3DRenderer.cpp b/Engine/Source/Runtime/SlateRHIRenderer/Private/Slate3DRenderer.cpp
--- a/Engine/Source/Runtime/SlateRHIRenderer/Private/Slate3DRenderer.cpp
+++ b/Engine/Source/Runtime/SlateRHIRenderer/Private/Slate3DRenderer.cpp
@@ -133,6 +133,8 @@ void FSlate3DRenderer::DrawWindowToTarget_RenderThread(FRHICommandListImmediate&
 
 	const TArray<TSharedRef<FSlateWindowElementList>>& WindowsToDraw = Context.WindowDrawBuffer->GetWindowElementLists();
 
+	FMemMark MemMark(FMemStack::Get());
+
 	// Enqueue a command to unlock the draw buffer after all windows have been drawn
 	RenderTargetPolicy->BeginDrawingWindows();
 
