Fix crash when applying a material in UMG to a brush using the rounded box style

diff --git a/Engine/Source/Runtime/SlateRHIRenderer/Private/SlateRHIRenderingPolicy.cpp b/Engine/Source/Runtime/SlateRHIRenderer/Private/SlateRHIRenderingPolicy.cpp
--- a/Engine/Source/Runtime/SlateRHIRenderer/Private/SlateRHIRenderingPolicy.cpp
+++ b/Engine/Source/Runtime/SlateRHIRenderer/Private/SlateRHIRenderingPolicy.cpp
@@ -1414,6 +1414,10 @@ TShaderRef<FSlateMaterialShaderPS> FSlateRHIRenderingPolicy::GetMaterialPixelSha
 	case ESlateShader::Custom:
 		FoundShader = MaterialShaderMap->GetShader(&TSlateMaterialShaderPS<ESlateShader::Custom>::StaticType);
 		break;
+	case ESlateShader::RoundedBox:
+		// Todo rounded box not supported in materials currently
+		FoundShader = MaterialShaderMap->GetShader(&TSlateMaterialShaderPS<ESlateShader::Default>::StaticType);
+		break;
 	default:
 		checkf(false, TEXT("Unsupported Slate shader type for use with materials"));
 		break;
