Fixed crash in UI material when scene texture node is used.

diff --git a/Engine/Source/Runtime/SlateRHIRenderer/Private/SlateMaterialShader.cpp b/Engine/Source/Runtime/SlateRHIRenderer/Private/SlateMaterialShader.cpp
--- a/Engine/Source/Runtime/SlateRHIRenderer/Private/SlateMaterialShader.cpp
+++ b/Engine/Source/Runtime/SlateRHIRenderer/Private/SlateMaterialShader.cpp
@@ -20,7 +20,8 @@ void FSlateMaterialShaderVS::ModifyCompilationEnvironment(const FMaterialShaderP
 	// Set defines based on what this shader will be used for
 	OutEnvironment.SetDefine( TEXT("USE_MATERIALS"), 1 );
 	OutEnvironment.SetDefine( TEXT("NUM_CUSTOMIZED_UVS"), Parameters.MaterialParameters.NumCustomizedUVs );
-	OutEnvironment.SetDefine( TEXT("HAS_SCREEN_POSITION"), (bool)Parameters.MaterialParameters.bHasVertexPositionOffsetConnected );
+	OutEnvironment.SetDefine(TEXT("HAS_SCREEN_POSITION"), (bool)Parameters.MaterialParameters.bHasVertexPositionOffsetConnected);
+	OutEnvironment.SetDefine(TEXT("SCENE_TEXTURES_DISABLED"), 1);
 
 	FMaterialShader::ModifyCompilationEnvironment( Parameters, OutEnvironment );
 }
@@ -47,18 +48,6 @@ void FSlateMaterialShaderVS::SetVerticalAxisMultiplier(FRHICommandList& RHICmdLi
 	SetShaderValue(RHICmdList, RHICmdList.GetBoundVertexShader(), SwitchVerticalAxisMultiplier, InMultiplier );
 }
 
-
-/*bool FSlateMaterialShaderVS::Serialize(FArchive& Ar)
-{
-	bool bShaderHasOutdatedParameters = FMaterialShader::Serialize(Ar);
-
-	Ar << ViewProjection;
-	Ar << SwitchVerticalAxisMultiplier;
-
-	return bShaderHasOutdatedParameters;
-}*/
-
-
 bool FSlateMaterialShaderPS::ShouldCompilePermutation(const FMaterialShaderPermutationParameters& Parameters)
 {
 	return Parameters.MaterialParameters.MaterialDomain == MD_UI;
@@ -70,6 +59,7 @@ void FSlateMaterialShaderPS::ModifyCompilationEnvironment(const FMaterialShaderP
 	// Set defines based on what this shader will be used for
 	OutEnvironment.SetDefine( TEXT("USE_MATERIALS"), 1 );
 	OutEnvironment.SetDefine( TEXT("NUM_CUSTOMIZED_UVS"), Parameters.MaterialParameters.NumCustomizedUVs);
+	OutEnvironment.SetDefine( TEXT("SCENE_TEXTURES_DISABLED"), 1);
 
 	FMaterialShader::ModifyCompilationEnvironment( Parameters, OutEnvironment );
 }
diff --git a/Engine/Source/Runtime/SlateRHIRenderer/Private/SlateMaterialShader.h b/Engine/Source/Runtime/SlateRHIRenderer/Private/SlateMaterialShader.h
--- a/Engine/Source/Runtime/SlateRHIRenderer/Private/SlateMaterialShader.h
+++ b/Engine/Source/Runtime/SlateRHIRenderer/Private/SlateMaterialShader.h
@@ -61,7 +61,7 @@ public:
 
 	void SetBlendState(FGraphicsPipelineStateInitializer& GraphicsPSOInit, const FMaterial* Material);
 
-	void SetParameters(FRHICommandList& RHICmdList, const FSceneView& View, const FMaterialRenderProxy* MaterialRenderProxy, const FMaterial* Material, const FVector4& InShaderParams);
+	void SetParameters(FRHICommandList& RHICmdList, const FSceneView& View, const FMaterialRenderProxy* MaterialRenderProxy, const FMaterial* Material, const FShaderParams& InShaderParams);
 
 	void SetDisplayGammaAndContrast(FRHICommandList& RHICmdList, float InDisplayGamma, float InContrast);
 
@@ -73,6 +73,7 @@ private:
 	LAYOUT_FIELD(FShaderParameter, GammaAndAlphaValues);
 	LAYOUT_FIELD(FShaderParameter, DrawFlags);
 	LAYOUT_FIELD(FShaderParameter, ShaderParams);
+	LAYOUT_FIELD(FShaderParameter, ShaderParams2);
 	/** Extra texture (like a font atlas) to be used in addition to any material textures */
 	LAYOUT_FIELD(FShaderResourceParameter, TextureParameterSampler);
 	LAYOUT_FIELD(FShaderResourceParameter, AdditionalTextureParameter);
@@ -102,6 +103,7 @@ public:
 		FSlateMaterialShaderVS::ModifyCompilationEnvironment(Parameters, OutEnvironment);
 
 		OutEnvironment.SetDefine(TEXT("USE_SLATE_INSTANCING"), (uint32)( bUseInstancing ? 1 : 0 ));
+		OutEnvironment.SetDefine(TEXT("SCENE_TEXTURES_DISABLED"), 1);
 	}
 };
 
@@ -129,5 +131,6 @@ public:
 		FSlateMaterialShaderPS::ModifyCompilationEnvironment(Parameters,OutEnvironment);
 
 		OutEnvironment.SetDefine(TEXT("SHADER_TYPE"), (uint32)ShaderType);
+		OutEnvironment.SetDefine(TEXT("SCENE_TEXTURES_DISABLED"), 1);
 	}
 };
