Fixed missing transitions on Slate depth buffers.

Also added names for said depth buffers, so we get meaningful error messages about them.

#rnx

diff --git a/Engine/Source/Runtime/SlateRHIRenderer/Private/Slate3DRenderer.cpp b/Engine/Source/Runtime/SlateRHIRenderer/Private/Slate3DRenderer.cpp
--- a/Engine/Source/Runtime/SlateRHIRenderer/Private/Slate3DRenderer.cpp
+++ b/Engine/Source/Runtime/SlateRHIRenderer/Private/Slate3DRenderer.cpp
@@ -168,6 +168,7 @@ void FSlate3DRenderer::DrawWindowToTarget_RenderThread(FRHICommandListImmediate&
 
 						FTexture2DRHIRef ShaderResourceUnused;
 						FRHIResourceCreateInfo CreateInfo(FClearValueBinding::DepthZero);
+						CreateInfo.DebugName = TEXT("SlateWindowDepthStencil");
 						RHICreateTargetableShaderResource2D(ColorTarget->GetSizeX(), ColorTarget->GetSizeY(), PF_DepthStencil, 1, TexCreate_None, TexCreate_DepthStencilTargetable, false, CreateInfo, DepthStencil, ShaderResourceUnused);
 						check(IsValidRef(DepthStencil));
 					}
diff --git a/Engine/Source/Runtime/SlateRHIRenderer/Private/SlateRHIRenderer.cpp b/Engine/Source/Runtime/SlateRHIRenderer/Private/SlateRHIRenderer.cpp
--- a/Engine/Source/Runtime/SlateRHIRenderer/Private/SlateRHIRenderer.cpp
+++ b/Engine/Source/Runtime/SlateRHIRenderer/Private/SlateRHIRenderer.cpp
@@ -151,6 +151,7 @@ void FViewportInfo::RecreateDepthBuffer_RenderThread()
 	{
 		FTexture2DRHIRef ShaderResourceUnused;
 		FRHIResourceCreateInfo CreateInfo(FClearValueBinding::DepthZero);
+		CreateInfo.DebugName = TEXT("SlateViewportDepthStencil");
 
 		ETextureCreateFlags TargetableTextureFlags = TexCreate_DepthStencilTargetable;
 		if (CVarMemorylessDepthStencil.GetValueOnAnyThread() != 0)
@@ -841,7 +842,6 @@ void FSlateRHIRenderer::DrawWindow_RenderThread(FRHICommandListImmediate& RHICmd
 
 			RHICmdList.BeginDrawingViewport(ViewportInfo.ViewportRHI, FTextureRHIRef());
 			RHICmdList.SetViewport(0, 0, 0, ViewportWidth, ViewportHeight, 0.0f);
-			RHICmdList.Transition(FRHITransitionInfo(BackBuffer, ERHIAccess::Unknown, ERHIAccess::RTV));
 
 			{
 				FRHIRenderPassInfo RPInfo(BackBuffer, ERenderTargetActions::Load_Store);
@@ -879,6 +879,7 @@ void FSlateRHIRenderer::DrawWindow_RenderThread(FRHICommandListImmediate& RHICmd
 					bool bHasBatches = BatchData.GetRenderBatches().Num() > 0;
 					if (bHasBatches || bClear)
 					{
+						TransitionRenderPassTargets(RHICmdList, RPInfo);
 						RHICmdList.BeginRenderPass(RPInfo, TEXT("SlateBatches"));
 						SCOPE_CYCLE_COUNTER(STAT_SlateRTDrawBatches);
 
@@ -931,8 +932,8 @@ void FSlateRHIRenderer::DrawWindow_RenderThread(FRHICommandListImmediate& RHICmd
 				if (bLUTStale)
 				{
 					// #todo-renderpasses will this touch every pixel? use NoAction?
-					RHICmdList.Transition(FRHITransitionInfo(ViewportInfo.ColorSpaceLUTRT, ERHIAccess::Unknown, ERHIAccess::RTV));
 					FRHIRenderPassInfo RPInfo(ViewportInfo.ColorSpaceLUTRT, ERenderTargetActions::Load_Store);
+					TransitionRenderPassTargets(RHICmdList, RPInfo);
 					RHICmdList.BeginRenderPass(RPInfo, TEXT("GenerateLUT"));
 					{
 						FGraphicsPipelineStateInitializer GraphicsPSOInit;
@@ -1053,6 +1054,7 @@ void FSlateRHIRenderer::DrawWindow_RenderThread(FRHICommandListImmediate& RHICmd
 				auto ShaderMap = GetGlobalShaderMap(FeatureLevel);
 
 				FRHIRenderPassInfo RPInfo(FinalBuffer, ERenderTargetActions::Load_Store);
+				TransitionRenderPassTargets(RHICmdList, RPInfo);
 				RHICmdList.BeginRenderPass(RPInfo, TEXT("SlateComposite"));
 
 				FGraphicsPipelineStateInitializer GraphicsPSOInit;
diff --git a/Engine/Source/Runtime/SlateRHIRenderer/Private/SlateRHIRenderingPolicy.cpp b/Engine/Source/Runtime/SlateRHIRenderer/Private/SlateRHIRenderingPolicy.cpp
--- a/Engine/Source/Runtime/SlateRHIRenderer/Private/SlateRHIRenderingPolicy.cpp
+++ b/Engine/Source/Runtime/SlateRHIRenderer/Private/SlateRHIRenderingPolicy.cpp
@@ -415,6 +415,7 @@ static bool UpdateScissorRect(
 					RPInfo.DepthStencilRenderTarget.Action = MakeDepthStencilTargetActions(ERenderTargetActions::DontLoad_DontStore, StencilAction);
 					RPInfo.DepthStencilRenderTarget.DepthStencilTarget = DepthStencilTarget;
 					RPInfo.DepthStencilRenderTarget.ExclusiveDepthStencil = FExclusiveDepthStencil::DepthNop_StencilWrite;
+					TransitionRenderPassTargets(RHICmdList, RPInfo);
 					RHICmdList.BeginRenderPass(RPInfo, TEXT("SlateUpdateScissorRect_ClearStencil"));
 				}
 
