Fixes for buffers being created or locked during render passes.

#rnx

diff --git a/Engine/Source/Runtime/Renderer/Private/VolumetricCloudRendering.cpp b/Engine/Source/Runtime/Renderer/Private/VolumetricCloudRendering.cpp
--- a/Engine/Source/Runtime/Renderer/Private/VolumetricCloudRendering.cpp
+++ b/Engine/Source/Runtime/Renderer/Private/VolumetricCloudRendering.cpp
@@ -891,17 +891,8 @@ static FSingleTriangleMeshVertexFactory* GSingleTriangleMeshVertexFactory = NULL
 
 static void GetSingleTriangleMeshBatch(FMeshBatch& LocalSingleTriangleMesh, const FMaterialRenderProxy* CloudVolumeMaterialProxy, const ERHIFeatureLevel::Type FeatureLevel)
 {
-	if (!GSingleTriangleMeshVertexFactory || GSingleTriangleMeshVertexFactory->HasIncompatibleFeatureLevel(FeatureLevel))
-	{
-		if (GSingleTriangleMeshVertexFactory)
-		{
-			GSingleTriangleMeshVertexFactory->ReleaseResource();
-			delete GSingleTriangleMeshVertexFactory;
-		}
-		GSingleTriangleMeshVertexFactory = new FSingleTriangleMeshVertexFactory(FeatureLevel);
-		GSingleTriangleMeshVertexBuffer.UpdateRHI();
-		GSingleTriangleMeshVertexFactory->InitResource();
-	}
+	check(GSingleTriangleMeshVertexFactory && !GSingleTriangleMeshVertexFactory->HasIncompatibleFeatureLevel(FeatureLevel));
+
 	LocalSingleTriangleMesh.VertexFactory = GSingleTriangleMeshVertexFactory;
 	LocalSingleTriangleMesh.MaterialRenderProxy = CloudVolumeMaterialProxy;
 	LocalSingleTriangleMesh.Elements[0].IndexBuffer = nullptr;
@@ -1326,6 +1317,18 @@ void FSceneRenderer::InitVolumetricCloudsForViews(FRDGBuilder& GraphBuilder, boo
 		return;
 	}
 
+	if (!GSingleTriangleMeshVertexFactory || GSingleTriangleMeshVertexFactory->HasIncompatibleFeatureLevel(ViewFamily.GetFeatureLevel()))
+	{
+		if (GSingleTriangleMeshVertexFactory)
+		{
+			GSingleTriangleMeshVertexFactory->ReleaseResource();
+			delete GSingleTriangleMeshVertexFactory;
+		}
+		GSingleTriangleMeshVertexFactory = new FSingleTriangleMeshVertexFactory(ViewFamily.GetFeatureLevel());
+		GSingleTriangleMeshVertexBuffer.UpdateRHI();
+		GSingleTriangleMeshVertexFactory->InitResource();
+	}
+
 	if (Scene)
 	{
 		check(ShouldRenderVolumetricCloud(Scene, ViewFamily.EngineShowFlags)); // This should not be called if we should not render SkyAtmosphere
diff --git a/Engine/Source/Runtime/SlateRHIRenderer/Private/Slate3DRenderer.cpp b/Engine/Source/Runtime/SlateRHIRenderer/Private/Slate3DRenderer.cpp
--- a/Engine/Source/Runtime/SlateRHIRenderer/Private/Slate3DRenderer.cpp
+++ b/Engine/Source/Runtime/SlateRHIRenderer/Private/Slate3DRenderer.cpp
@@ -132,64 +132,63 @@ void FSlate3DRenderer::DrawWindowToTarget_RenderThread(FRHICommandListImmediate&
 	{
 		RPInfo.ColorRenderTargets[0].Action = ERenderTargetActions::Clear_Store;
 	}
-	InRHICmdList.BeginRenderPass(RPInfo, TEXT("Slate3D"));
-	{
 
-		for (int32 WindowIndex = 0; WindowIndex < WindowsToDraw.Num(); WindowIndex++)
-		{
-			FSlateWindowElementList& WindowElementList = *WindowsToDraw[WindowIndex];
+	for (int32 WindowIndex = 0; WindowIndex < WindowsToDraw.Num(); WindowIndex++)
+	{
+		FSlateWindowElementList& WindowElementList = *WindowsToDraw[WindowIndex];
 
-			FSlateBatchData& BatchData = WindowElementList.GetBatchData();
+		FSlateBatchData& BatchData = WindowElementList.GetBatchData();
 
-			if (BatchData.GetRenderBatches().Num() > 0)
-			{
-				RenderTargetPolicy->BuildRenderingBuffers(InRHICmdList, BatchData);
+		if (BatchData.GetRenderBatches().Num() > 0)
+		{
+			RenderTargetPolicy->BuildRenderingBuffers(InRHICmdList, BatchData);
 		
-				FVector2D DrawOffset = Context.WindowDrawBuffer->ViewOffset;
+			FVector2D DrawOffset = Context.WindowDrawBuffer->ViewOffset;
 
-				FMatrix ProjectionMatrix = FSlateRHIRenderer::CreateProjectionMatrix(RTTextureRHI->GetSizeX(), RTTextureRHI->GetSizeY());
-				FMatrix ViewOffset = FTranslationMatrix::Make(FVector(DrawOffset, 0));
-				ProjectionMatrix = ViewOffset * ProjectionMatrix;
+			FMatrix ProjectionMatrix = FSlateRHIRenderer::CreateProjectionMatrix(RTTextureRHI->GetSizeX(), RTTextureRHI->GetSizeY());
+			FMatrix ViewOffset = FTranslationMatrix::Make(FVector(DrawOffset, 0));
+			ProjectionMatrix = ViewOffset * ProjectionMatrix;
 
-				FSlateBackBuffer BackBufferTarget(Context.RenderTarget->GetRenderTargetTexture(), FIntPoint(RTTextureRHI->GetSizeX(), RTTextureRHI->GetSizeY()));
+			FSlateBackBuffer BackBufferTarget(Context.RenderTarget->GetRenderTargetTexture(), FIntPoint(RTTextureRHI->GetSizeX(), RTTextureRHI->GetSizeY()));
 
-				FSlateRenderingParams DrawOptions(ProjectionMatrix, Context.WorldTimeSeconds, Context.DeltaTimeSeconds, Context.RealTimeSeconds);
-				// The scene renderer will handle it in this case
-				DrawOptions.bAllowSwitchVerticalAxis = false;
-				DrawOptions.ViewOffset = DrawOffset;
+			FSlateRenderingParams DrawOptions(ProjectionMatrix, Context.WorldTimeSeconds, Context.DeltaTimeSeconds, Context.RealTimeSeconds);
+			// The scene renderer will handle it in this case
+			DrawOptions.bAllowSwitchVerticalAxis = false;
+			DrawOptions.ViewOffset = DrawOffset;
 
-				FTexture2DRHIRef ColorTarget = Context.RenderTarget->GetRenderTargetTexture();
+			FTexture2DRHIRef ColorTarget = Context.RenderTarget->GetRenderTargetTexture();
 
-				if (BatchData.IsStencilClippingRequired())
+			if (BatchData.IsStencilClippingRequired())
+			{
+				if (!DepthStencil.IsValid() || ColorTarget->GetSizeXY() != DepthStencil->GetSizeXY())
 				{
-					if (!DepthStencil.IsValid() || ColorTarget->GetSizeXY() != DepthStencil->GetSizeXY())
-					{
-						DepthStencil.SafeRelease();
-
-						FTexture2DRHIRef ShaderResourceUnused;
-						FRHIResourceCreateInfo CreateInfo(TEXT("SlateWindowDepthStencil"), FClearValueBinding::DepthZero);
-						RHICreateTargetableShaderResource2D(ColorTarget->GetSizeX(), ColorTarget->GetSizeY(), PF_DepthStencil, 1, TexCreate_None, TexCreate_DepthStencilTargetable, false, CreateInfo, DepthStencil, ShaderResourceUnused);
-						check(IsValidRef(DepthStencil));
-					}
-				}
+					DepthStencil.SafeRelease();
 
-				RenderTargetPolicy->DrawElements(
-					InRHICmdList,
-					BackBufferTarget,
-					ColorTarget,
-					ColorTarget,
-					DepthStencil,
-					BatchData.GetFirstRenderBatchIndex(),
-					BatchData.GetRenderBatches(),
-					DrawOptions
-				);
+					FTexture2DRHIRef ShaderResourceUnused;
+					FRHIResourceCreateInfo CreateInfo(TEXT("SlateWindowDepthStencil"), FClearValueBinding::DepthZero);
+					RHICreateTargetableShaderResource2D(ColorTarget->GetSizeX(), ColorTarget->GetSizeY(), PF_DepthStencil, 1, TexCreate_None, TexCreate_DepthStencilTargetable, false, CreateInfo, DepthStencil, ShaderResourceUnused);
+					check(IsValidRef(DepthStencil));
+				}
 			}
+
+			// Ideally we'd have a single render pass for all the windows, but this code reuses a single vertex buffer for each draw, which it updates above in BuildRenderingBuffers(),
+			// and we can't upload data during a render pass. We need to rewrite this to upload all the data to the buffer first and use offsets for each draw.
+			InRHICmdList.BeginRenderPass(RPInfo, TEXT("Slate3D"));
+
+			RenderTargetPolicy->DrawElements(
+				InRHICmdList,
+				BackBufferTarget,
+				ColorTarget,
+				ColorTarget,
+				DepthStencil,
+				BatchData.GetFirstRenderBatchIndex(),
+				BatchData.GetRenderBatches(),
+				DrawOptions
+			);
+
+			InRHICmdList.EndRenderPass();
 		}
 	}
-	if(InRHICmdList.IsInsideRenderPass())
-	{
-		InRHICmdList.EndRenderPass();
-	}
 
 	FSlateEndDrawingWindowsCommand::EndDrawingWindows(InRHICmdList, Context.WindowDrawBuffer, *RenderTargetPolicy);
 	InRHICmdList.CopyToResolveTarget(Context.RenderTarget->GetRenderTargetTexture(), RTTextureRHI, FResolveParams());
