Fix crash due to Slate messing with render passes.

#rnx

diff --git a/Engine/Source/Runtime/SlateRHIRenderer/Private/Slate3DRenderer.cpp b/Engine/Source/Runtime/SlateRHIRenderer/Private/Slate3DRenderer.cpp
--- a/Engine/Source/Runtime/SlateRHIRenderer/Private/Slate3DRenderer.cpp
+++ b/Engine/Source/Runtime/SlateRHIRenderer/Private/Slate3DRenderer.cpp
@@ -186,7 +186,11 @@ void FSlate3DRenderer::DrawWindowToTarget_RenderThread(FRHICommandListImmediate&
 				DrawOptions
 			);
 
-			InRHICmdList.EndRenderPass();
+			// FSlateRHIRenderingPolicy::DrawElements can close the render pass on its own sometimes, so don't do it again.
+			if (InRHICmdList.IsInsideRenderPass())
+			{
+				InRHICmdList.EndRenderPass();
+			}
 		}
 	}
 
