Add missing call to UpdateDeferredCachedUniformExpressions when doing slate rendering

diff --git a/Engine/Source/Runtime/SlateRHIRenderer/Private/Slate3DRenderer.cpp b/Engine/Source/Runtime/SlateRHIRenderer/Private/Slate3DRenderer.cpp
--- a/Engine/Source/Runtime/SlateRHIRenderer/Private/Slate3DRenderer.cpp
+++ b/Engine/Source/Runtime/SlateRHIRenderer/Private/Slate3DRenderer.cpp
@@ -120,6 +120,8 @@ void FSlate3DRenderer::DrawWindowToTarget_RenderThread(FRHICommandListImmediate&
 
 	FMemMark MemMark(FMemStack::Get());
 
+	FMaterialRenderProxy::UpdateDeferredCachedUniformExpressions();
+
 	// Enqueue a command to unlock the draw buffer after all windows have been drawn
 	RenderTargetPolicy->BeginDrawingWindows();
 
