Fix to Slate3DRenderer to do a clear even if the RenderBatch is empty. This can happens for example when a Widget Component with the World space put all widget invisible. Whithout this fix, we end up not clearing the rendering and we keep the last render.

diff --git a/Engine/Source/Runtime/SlateRHIRenderer/Private/Slate3DRenderer.cpp b/Engine/Source/Runtime/SlateRHIRenderer/Private/Slate3DRenderer.cpp
--- a/Engine/Source/Runtime/SlateRHIRenderer/Private/Slate3DRenderer.cpp
+++ b/Engine/Source/Runtime/SlateRHIRenderer/Private/Slate3DRenderer.cpp
@@ -12,6 +12,7 @@ DECLARE_GPU_STAT_NAMED(Slate3D, TEXT("Slate 3D"));
 FSlate3DRenderer::FSlate3DRenderer( TSharedRef<FSlateFontServices> InSlateFontServices, TSharedRef<FSlateRHIResourceManager> InResourceManager, bool bUseGammaCorrection )
 	: SlateFontServices( InSlateFontServices )
 	, ResourceManager( InResourceManager )
+	, bRenderTargetWasCleared(false)
 {
 	const int32 InitialBufferSize = 200;
 	RenderTargetPolicy = MakeShareable( new FSlateRHIRenderingPolicy( SlateFontServices, ResourceManager, InitialBufferSize ) );
@@ -197,6 +198,19 @@ void FSlate3DRenderer::DrawWindowToTarget_RenderThread(FRHICommandListImmediate&
 			{
 				InRHICmdList.EndRenderPass();
 			}
+
+			// each time we do draw content, we reset the flag
+			bRenderTargetWasCleared = false;
+		}
+		// If we have no render command and it's the first clear, we call Begin End to force the clear of the render target. Otherwise, we end up not updating the buffer when all Widget are invisible.
+		else if(!bRenderTargetWasCleared)
+		{
+			InRHICmdList.BeginRenderPass(RPInfo, TEXT("Slate3D")); 
+			if (InRHICmdList.IsInsideRenderPass())
+			{
+				InRHICmdList.EndRenderPass();
+			}
+			bRenderTargetWasCleared = true;
 		}
 	}
 
diff --git a/Engine/Source/Runtime/SlateRHIRenderer/Private/Slate3DRenderer.h b/Engine/Source/Runtime/SlateRHIRenderer/Private/Slate3DRenderer.h
--- a/Engine/Source/Runtime/SlateRHIRenderer/Private/Slate3DRenderer.h
+++ b/Engine/Source/Runtime/SlateRHIRenderer/Private/Slate3DRenderer.h
@@ -50,4 +50,7 @@ private:
 
 	/** The depth buffer texture if any */
 	FTexture2DRHIRef DepthStencil;
+
+	/** Set to true when the render target was cleared and prevent another clear to be called for nothing.**/
+	bool bRenderTargetWasCleared;
 };
