PR #8373: Fix slate post processor (Contributed by [REDACTED_PERSON])

diff --git a/Engine/Source/Runtime/SlateRHIRenderer/Private/SlatePostProcessResource.cpp b/Engine/Source/Runtime/SlateRHIRenderer/Private/SlatePostProcessResource.cpp
--- a/Engine/Source/Runtime/SlateRHIRenderer/Private/SlatePostProcessResource.cpp
+++ b/Engine/Source/Runtime/SlateRHIRenderer/Private/SlatePostProcessResource.cpp
@@ -24,8 +24,9 @@ void FSlatePostProcessResource::Update(const FIntPoint& NewSize)
 		{
 			InitResource();
 		}
-	
-		ResizeTargets(NewSize);
+
+		FIntPoint NewMaxSize(FMath::Max(NewSize.X, RenderTargetSize.X), FMath::Max(NewSize.Y, RenderTargetSize.Y));
+		ResizeTargets(NewMaxSize);
 	}
 }
 
