[NO NON-PII COMMIT MESSAGE]

diff --git a/Engine/Source/Runtime/SlateRHIRenderer/Private/Slate3DRenderer.cpp b/Engine/Source/Runtime/SlateRHIRenderer/Private/Slate3DRenderer.cpp
--- a/Engine/Source/Runtime/SlateRHIRenderer/Private/Slate3DRenderer.cpp
+++ b/Engine/Source/Runtime/SlateRHIRenderer/Private/Slate3DRenderer.cpp
@@ -116,6 +116,10 @@ void FSlate3DRenderer::DrawWindowToTarget_RenderThread(FRHICommandListImmediate&
 
 	checkSlow(Context.RenderTarget);
 
+	//Update cached uniforms to avoid an ensure being thrown due to a null shader map (UE-110263)
+	//Same fix on UE5 stream - CL16165057 
+	FMaterialRenderProxy::UpdateDeferredCachedUniformExpressions();
+
 	const TArray<TSharedRef<FSlateWindowElementList>>& WindowsToDraw = Context.WindowDrawBuffer->GetWindowElementLists();
 
 	FMemMark MemMark(FMemStack::Get());
