Slate: Clear the uncached draw elements once they are batches and clear the batches elements once they are rendered. We do this to remove hard references to the Glyph Sequence from the previous frame. Cached draw elements are still cached. An old Glyph Sequence may still be a reference in a cached element but it is the responsibility of the SWidget to do the proper invalidation.

diff --git a/Engine/Source/Runtime/SlateCore/Private/Rendering/DrawElements.cpp b/Engine/Source/Runtime/SlateCore/Private/Rendering/DrawElements.cpp
--- a/Engine/Source/Runtime/SlateCore/Private/Rendering/DrawElements.cpp
+++ b/Engine/Source/Runtime/SlateCore/Private/Rendering/DrawElements.cpp
@@ -970,16 +970,32 @@ DECLARE_DWORD_COUNTER_STAT(TEXT("FSlateWindowElementList MemManager Count"), STA
 
 void FSlateWindowElementList::ResetElementList()
 {
+	ResetDrawElementList();
+	ResetBatchData();
+}
+
+void FSlateWindowElementList::ResetBatchData()
+{
+	// data that is used on the render thread
+	BatchData.ResetData();
+	ClippingManager.ResetClippingState();
+
+	check(DeferredPaintList.Num() == 0);
+	check(UncachedDrawElements.Num() == 0);
+	check(CachedElementDataList.Num() == 0);
+	check(CachedElementDataListStack.Num() == 0);
+	check(RenderTargetWindow == nullptr);
+}
+
+void FSlateWindowElementList::ResetDrawElementList()
+{
+	// data that is used on the game thread
 	QUICK_SCOPE_CYCLE_COUNTER(Slate_ResetElementList);
 
 	check(IsThreadSafeForSlateRendering());
 
 	DeferredPaintList.Reset();
 
-	BatchData.ResetData();
-
-	ClippingManager.ResetClippingState();
-
 	UncachedDrawElements.Reset();
 
 #if STATS
diff --git a/Engine/Source/Runtime/SlateCore/Private/Rendering/ElementBatcher.cpp b/Engine/Source/Runtime/SlateCore/Private/Rendering/ElementBatcher.cpp
--- a/Engine/Source/Runtime/SlateCore/Private/Rendering/ElementBatcher.cpp
+++ b/Engine/Source/Runtime/SlateCore/Private/Rendering/ElementBatcher.cpp
@@ -312,6 +312,9 @@ void FSlateElementBatcher::AddElements(FSlateWindowElementList& WindowElementLis
 		}
 	}
 
+	// Draw Elements are added. They are not needed anymore.
+	WindowElementList.ResetDrawElementList();
+
 	// Done with the element list
 	BatchData = nullptr;
 	PrecachedClippingStates = nullptr;
diff --git a/Engine/Source/Runtime/SlateCore/Private/Rendering/SlateDrawBuffer.cpp b/Engine/Source/Runtime/SlateCore/Private/Rendering/SlateDrawBuffer.cpp
--- a/Engine/Source/Runtime/SlateCore/Private/Rendering/SlateDrawBuffer.cpp
+++ b/Engine/Source/Runtime/SlateCore/Private/Rendering/SlateDrawBuffer.cpp
@@ -25,14 +25,15 @@ FSlateWindowElementList& FSlateDrawBuffer::AddWindowElementList(TSharedRef<SWind
 			WindowElementLists.Add(ExistingElementList);
 			WindowElementListsPool.RemoveAtSwap(WindowIndex);
 
-			ExistingElementList->ResetElementList();
-
+			ensureMsgf(ExistingElementList->GetBatchData().GetNumFinalBatches() == 0, TEXT("The Buffer should have been clear when it was unlocked."));
 			return *ExistingElementList;
 		}
 	}
 
 	TSharedRef<FSlateWindowElementList> WindowElements = MakeShared<FSlateWindowElementList>(ForWindow);
-	WindowElements->ResetElementList();
+	WindowElements->ResetDrawElementList();
+	WindowElements->ResetBatchData();
+
 	WindowElementLists.Add(WindowElements);
 
 	return *WindowElements;
@@ -46,6 +47,8 @@ void FSlateDrawBuffer::RemoveUnusedWindowElement(const TArray<SWindow*>& AllWind
 		SWindow* CandidateWindow = WindowElementLists[WindowIndex]->GetPaintWindow();
 		if (!CandidateWindow || !AllWindows.Contains(CandidateWindow))
 		{
+			WindowElementLists[WindowIndex]->ResetDrawElementList();
+			WindowElementLists[WindowIndex]->ResetBatchData();
 			WindowElementLists.RemoveAtSwap(WindowIndex);
 			--WindowIndex;
 		}
@@ -54,18 +57,31 @@ void FSlateDrawBuffer::RemoveUnusedWindowElement(const TArray<SWindow*>& AllWind
 
 bool FSlateDrawBuffer::Lock()
 {
-	return FPlatformAtomics::InterlockedCompareExchange(&Locked, 1, 0) == 0;
+	bool ExpectedValue = false;
+	bool bIsLock = bLocked.compare_exchange_strong(ExpectedValue, true);
+	if (bIsLock)
+	{
+		bIsLockedBySlateThread = IsInSlateThread();
+	}
+	return bIsLock;
 }
 
 void FSlateDrawBuffer::Unlock()
 {
-	FPlatformAtomics::InterlockedExchange(&Locked, 0);
+	// Rendering doesn't need the batch data anymore
+	for (TSharedRef<FSlateWindowElementList>& ExistingElementList : WindowElementLists)
+	{
+		ExistingElementList->ResetBatchData();
+	}
+
+	bLocked = false;
 }
 
 void FSlateDrawBuffer::AddReferencedObjects(FReferenceCollector& Collector)
 {
-	// Locked buffers are the only ones that are currently referencing objects.  If unlocked, the element list is not in-use and contains to-be-cleared data
-	if(Locked != 0)
+	// Locked buffers are the only ones that are currently referencing objects.
+	//If unlocked, the element list is not in-use and contains "to-be-cleared" data.
+	if(bLocked && !bIsLockedBySlateThread)
 	{
 		for (TSharedRef<FSlateWindowElementList>& ElementList : WindowElementLists)
 		{
diff --git a/Engine/Source/Runtime/SlateCore/Public/Rendering/DrawElements.h b/Engine/Source/Runtime/SlateCore/Public/Rendering/DrawElements.h
--- a/Engine/Source/Runtime/SlateCore/Public/Rendering/DrawElements.h
+++ b/Engine/Source/Runtime/SlateCore/Public/Rendering/DrawElements.h
@@ -313,6 +313,7 @@ private:
 
 	static FVector2D GetRotationPoint( const FPaintGeometry& PaintGeometry, const TOptional<FVector2D>& UserRotationPoint, ERotationSpace RotationSpace );
 	static FSlateDrawElement& MakeBoxInternal(FSlateWindowElementList& ElementList, uint32 InLayer, const FPaintGeometry& PaintGeometry, const FSlateBrush* InBrush, ESlateDrawEffect InDrawEffects, const FLinearColor& InTint);
+
 private:
 	FSlateDataPayload* DataPayload;
 	FSlateRenderTransform RenderTransform;
@@ -533,9 +534,7 @@ public:
 		return *Payload;
 	}
 
-	/**
-	 * Creates an uninitialized draw element
-	 */
+	/** Creates an uninitialized draw element */
 	SLATECORE_API FSlateDrawElement& AddUninitialized();
 
 
@@ -619,11 +618,16 @@ public:
 	// OTHER
 	//--------------------------------------------------------------------------
 	
-	/**
-	 * Remove all the elements from this draw list.
-	 */
+	/** Remove all the elements from this draw list. */
+	UE_DEPRECATED(5.0, "ResetElementList is deprecated. Use ResetBatchData or ResetDrawElementList to reset the elements")
 	SLATECORE_API void ResetElementList();
 
+	/** Remove all the batch data. */
+	SLATECORE_API void ResetBatchData();
+
+	/** Remove all the draw elements needed to create the batch data. */
+	SLATECORE_API void ResetDrawElementList();
+
 
 	FSlateBatchData& GetBatchData() { return BatchData; }
 
@@ -637,11 +641,12 @@ private:
 	TArrayView<FSlateCachedElementData* const> GetCachedElementDataList() const { return MakeArrayView(CachedElementDataList.GetData(), CachedElementDataList.Num()); }
 
 	FSlateCachedElementData* GetCurrentCachedElementData() const { return CachedElementDataListStack.Num() ? CachedElementDataList[CachedElementDataListStack.Top()] : nullptr; }
+
 private:
 	/**
-	* Window which owns the widgets that are being painted but not necessarily rendered to
-	* Widgets are always rendered to the RenderTargetWindow
-	*/
+	 * Window which owns the widgets that are being painted but not necessarily rendered to
+	 * Widgets are always rendered to the RenderTargetWindow
+	 */
 	TWeakPtr<SWindow> WeakPaintWindow;
 	SWindow* RawPaintWindow;
 
@@ -669,7 +674,7 @@ private:
 	bool bNeedsDeferredResolve;
 	TArray<int32> ResolveToDeferredIndex;
 
-	// Begin Fast Path
+	//~ Begin Fast Path
 
 	/** State of the current widget that is adding draw elements */
 	struct FWidgetDrawElementState
@@ -692,7 +697,7 @@ private:
 	TArray<FWidgetDrawElementState, TInlineAllocator<50>> WidgetDrawStack;
 	TArray<FSlateCachedElementData*, TInlineAllocator<4>> CachedElementDataList;
 	TArray<int32, TInlineAllocator<4>> CachedElementDataListStack;
-	// End Fast Path
+	//~ End Fast Path
 
 	/** Store the size of the window being used to paint */
 	FVector2D WindowSize;
diff --git a/Engine/Source/Runtime/SlateCore/Public/Rendering/SlateDrawBuffer.h b/Engine/Source/Runtime/SlateCore/Public/Rendering/SlateDrawBuffer.h
--- a/Engine/Source/Runtime/SlateCore/Public/Rendering/SlateDrawBuffer.h
+++ b/Engine/Source/Runtime/SlateCore/Public/Rendering/SlateDrawBuffer.h
@@ -5,6 +5,8 @@
 #include "CoreMinimal.h"
 #include "UObject/GCObject.h"
 
+#include <atomic>
+
 class FSlateWindowElementList;
 class SWindow;
 
@@ -17,7 +19,8 @@ public:
 
 	/** Default constructor. */
 	explicit FSlateDrawBuffer()
-		: Locked(0)
+		: bLocked(false)
+		, bIsLockedBySlateThread(false)
 		, ResourceVersion(0)
 	{ }
 
@@ -78,8 +81,8 @@ protected:
 	// that we restore if they're requested again.
 	TArray< TSharedRef<FSlateWindowElementList> > WindowElementListsPool;
 
-	// 1 if this buffer is locked, 0 otherwise.
-	volatile int32 Locked;
+	std::atomic<bool> bLocked;
+	bool bIsLockedBySlateThread;
 
 	// Last recorded version from the render. The WindowElementListsPool is emptied when this changes.
 	uint32 ResourceVersion;
diff --git a/Engine/Source/Runtime/SlateRHIRenderer/Private/SlateRHIRenderer.cpp b/Engine/Source/Runtime/SlateRHIRenderer/Private/SlateRHIRenderer.cpp
--- a/Engine/Source/Runtime/SlateRHIRenderer/Private/SlateRHIRenderer.cpp
+++ b/Engine/Source/Runtime/SlateRHIRenderer/Private/SlateRHIRenderer.cpp
@@ -91,8 +91,9 @@ TAutoConsoleVariable<int32> CVarShowSlateBatching(
 struct FSlateDrawWindowCommandParams
 {
 	FSlateRHIRenderer* Renderer;
-	FSlateWindowElementList* WindowElementList;
 	SWindow* Window;
+	FSlateBatchData* BatchData;
+	FVector2D WindowSize;
 #if WANTS_DRAW_MESH_EVENTS
 	FString WindowTitle;
 #endif
@@ -676,7 +677,7 @@ int32 SlateWireFrame = 0;
 static FAutoConsoleVariableRef CVarSlateWireframe(TEXT("Slate.ShowWireFrame"), SlateWireFrame, TEXT(""), ECVF_Default);
 
 /** Draws windows from a FSlateDrawBuffer on the render thread */
-void FSlateRHIRenderer::DrawWindow_RenderThread(FRHICommandListImmediate& RHICmdList, FViewportInfo& ViewportInfo, FSlateWindowElementList& WindowElementList, const struct FSlateDrawWindowCommandParams& DrawCommandParams)
+void FSlateRHIRenderer::DrawWindow_RenderThread(FRHICommandListImmediate& RHICmdList, FViewportInfo& ViewportInfo, struct FSlateDrawWindowCommandParams& DrawCommandParams)
 {
 	LLM_SCOPE(ELLMTag::SceneRender);
 	
@@ -730,7 +731,7 @@ void FSlateRHIRenderer::DrawWindow_RenderThread(FRHICommandListImmediate& RHICmd
 		bool bRenderedStereo = false;
 		if (CVarDrawToVRRenderTarget->GetInt() == 0 && GEngine && IsValidRef(ViewportInfo.GetRenderTargetTexture()) && GEngine->StereoRenderingDevice.IsValid())
 		{
-			const FVector2D WindowSize = WindowElementList.GetWindowSize();
+			const FVector2D WindowSize = DrawCommandParams.WindowSize;
 			GEngine->StereoRenderingDevice->RenderTexture_RenderThread(RHICmdList, RHICmdList.GetViewportBackBuffer(ViewportInfo.ViewportRHI), ViewportInfo.GetRenderTargetTexture(), WindowSize);
 			bRenderedStereo = true;
 		}
@@ -740,7 +741,8 @@ void FSlateRHIRenderer::DrawWindow_RenderThread(FRHICommandListImmediate& RHICmd
 			SCOPE_CYCLE_COUNTER(STAT_SlateRenderingRTTime);
 			CSV_SCOPED_TIMING_STAT_EXCLUSIVE(Slate);
 
-			FSlateBatchData& BatchData = WindowElementList.GetBatchData();
+			check(DrawCommandParams.BatchData);
+			FSlateBatchData& BatchData = *DrawCommandParams.BatchData;
 
 			// Update the vertex and index buffer	
 			RenderingPolicy->BuildRenderingBuffers(RHICmdList, BatchData);
@@ -1093,7 +1095,7 @@ void FSlateRHIRenderer::DrawWindow_RenderThread(FRHICommandListImmediate& RHICmd
 #endif
 			if (!bRenderedStereo && GEngine && IsValidRef(ViewportInfo.GetRenderTargetTexture()) && GEngine->StereoRenderingDevice.IsValid())
 			{
-				const FVector2D WindowSize = WindowElementList.GetWindowSize();
+				const FVector2D WindowSize = DrawCommandParams.WindowSize;
 				GEngine->StereoRenderingDevice->RenderTexture_RenderThread(RHICmdList, RHICmdList.GetViewportBackBuffer(ViewportInfo.ViewportRHI), ViewportInfo.GetRenderTargetTexture(), WindowSize);
 			}
 			RHICmdList.Transition(FRHITransitionInfo(BackBuffer, ERHIAccess::Unknown, ERHIAccess::SRVGraphics));
@@ -1311,7 +1313,8 @@ void FSlateRHIRenderer::DrawWindows_Private(FSlateDrawBuffer& WindowDrawBuffer)
 					FSlateDrawWindowCommandParams Params;
 
 					Params.Renderer = this;
-					Params.WindowElementList = &ElementList;
+					Params.WindowSize = ElementList.GetWindowSize();
+					Params.BatchData = &ElementList.GetBatchData();
 					Params.Window = Window;
 #if WANTS_DRAW_MESH_EVENTS
 					Params.WindowTitle = Window->GetTitle().ToString();
@@ -1331,9 +1334,9 @@ void FSlateRHIRenderer::DrawWindows_Private(FSlateDrawBuffer& WindowDrawBuffer)
 					if (GIsClient && !IsRunningCommandlet() && !GUsingNullRHI)
 					{
 						ENQUEUE_RENDER_COMMAND(SlateDrawWindowsCommand)(
-							[Params, ViewInfo](FRHICommandListImmediate& RHICmdList)
+							[Params, ViewInfo](FRHICommandListImmediate& RHICmdList) mutable
 							{
-								Params.Renderer->DrawWindow_RenderThread(RHICmdList, *ViewInfo, *Params.WindowElementList, Params);
+								Params.Renderer->DrawWindow_RenderThread(RHICmdList, *ViewInfo, Params);
 							}
 						);
 					}
diff --git a/Engine/Source/Runtime/SlateRHIRenderer/Private/SlateRHIRenderer.h b/Engine/Source/Runtime/SlateRHIRenderer/Private/SlateRHIRenderer.h
--- a/Engine/Source/Runtime/SlateRHIRenderer/Private/SlateRHIRenderer.h
+++ b/Engine/Source/Runtime/SlateRHIRenderer/Private/SlateRHIRenderer.h
@@ -214,7 +214,7 @@ public:
 	virtual void AddWidgetRendererUpdate(const struct FRenderThreadUpdateContext& Context, bool bDeferredRenderTargetUpdate) override;
 
 	/** Draws windows from a FSlateDrawBuffer on the render thread */
-	void DrawWindow_RenderThread(FRHICommandListImmediate& RHICmdList, FViewportInfo& ViewportInfo, FSlateWindowElementList& WindowElementList, const struct FSlateDrawWindowCommandParams& DrawParams);
+	void DrawWindow_RenderThread(FRHICommandListImmediate& RHICmdList, FViewportInfo& ViewportInfo, struct FSlateDrawWindowCommandParams& DrawParams);
 
 	/**
 	 * Reloads texture resources from disk                   
diff --git a/Engine/Source/Runtime/UMG/Private/Slate/WidgetRenderer.cpp b/Engine/Source/Runtime/UMG/Private/Slate/WidgetRenderer.cpp
--- a/Engine/Source/Runtime/UMG/Private/Slate/WidgetRenderer.cpp
+++ b/Engine/Source/Runtime/UMG/Private/Slate/WidgetRenderer.cpp
@@ -342,7 +342,7 @@ bool FWidgetRenderer::DrawInvalidationRoot(TSharedRef<SVirtualWindow>& VirtualWi
 		}
 		else
 		{
-			WindowElementList.ResetElementList();
+			WindowElementList.ResetDrawElementList();
 			DrawBuffer.Unlock();
 		}
 	}
