[Backout] - CL18952424

This is causing an ensure to hit in Slate applications.

Original CL Desc
-----------------------------------------------------------------
Slate: Clear the uncached draw elements once they are batches and clear the batches elements once they are rendered. We do this to remove hard references to the Glyph Sequence from the previous frame. Cached draw elements are still cached. An old Glyph Sequence may still be a reference in a cached element but it is the responsibility of the SWidget to do the proper invalidation.

diff --git a/Engine/Source/Runtime/SlateCore/Private/Rendering/DrawElements.cpp b/Engine/Source/Runtime/SlateCore/Private/Rendering/DrawElements.cpp
--- a/Engine/Source/Runtime/SlateCore/Private/Rendering/DrawElements.cpp
+++ b/Engine/Source/Runtime/SlateCore/Private/Rendering/DrawElements.cpp
@@ -970,32 +970,16 @@ DECLARE_DWORD_COUNTER_STAT(TEXT("FSlateWindowElementList MemManager Count"), STA
 
 void FSlateWindowElementList::ResetElementList()
 {
-	ResetDrawElementList();
-	ResetBatchData();
-}
-
-void FSlateWindowElementList::ResetBatchData()
-{
-	// data that is used on the render thread
-	BatchData.ResetData();
-	ClippingManager.ResetClippingState();
-
-	check(DeferredPaintList.Num() == 0);
-	check(UncachedDrawElements.Num() == 0);
-	check(CachedElementDataList.Num() == 0);
-	check(CachedElementDataListStack.Num() == 0);
-	check(RenderTargetWindow == nullptr);
-}
-
-void FSlateWindowElementList::ResetDrawElementList()
-{
-	// data that is used on the game thread
 	QUICK_SCOPE_CYCLE_COUNTER(Slate_ResetElementList);
 
 	check(IsThreadSafeForSlateRendering());
 
 	DeferredPaintList.Reset();
 
+	BatchData.ResetData();
+
+	ClippingManager.ResetClippingState();
+
 	UncachedDrawElements.Reset();
 
 #if STATS
diff --git a/Engine/Source/Runtime/SlateCore/Private/Rendering/ElementBatcher.cpp b/Engine/Source/Runtime/SlateCore/Private/Rendering/ElementBatcher.cpp
--- a/Engine/Source/Runtime/SlateCore/Private/Rendering/ElementBatcher.cpp
+++ b/Engine/Source/Runtime/SlateCore/Private/Rendering/ElementBatcher.cpp
@@ -312,9 +312,6 @@ void FSlateElementBatcher::AddElements(FSlateWindowElementList& WindowElementLis
 		}
 	}
 
-	// Draw Elements are added. They are not needed anymore.
-	WindowElementList.ResetDrawElementList();
-
 	// Done with the element list
 	BatchData = nullptr;
 	PrecachedClippingStates = nullptr;
diff --git a/Engine/Source/Runtime/SlateCore/Private/Rendering/SlateDrawBuffer.cpp b/Engine/Source/Runtime/SlateCore/Private/Rendering/SlateDrawBuffer.cpp
--- a/Engine/Source/Runtime/SlateCore/Private/Rendering/SlateDrawBuffer.cpp
+++ b/Engine/Source/Runtime/SlateCore/Private/Rendering/SlateDrawBuffer.cpp
@@ -25,15 +25,14 @@ FSlateWindowElementList& FSlateDrawBuffer::AddWindowElementList(TSharedRef<SWind
 			WindowElementLists.Add(ExistingElementList);
 			WindowElementListsPool.RemoveAtSwap(WindowIndex);
 
-			ensureMsgf(ExistingElementList->GetBatchData().GetNumFinalBatches() == 0, TEXT("The Buffer should have been clear when it was unlocked."));
+			ExistingElementList->ResetElementList();
+
 			return *ExistingElementList;
 		}
 	}
 
 	TSharedRef<FSlateWindowElementList> WindowElements = MakeShared<FSlateWindowElementList>(ForWindow);
-	WindowElements->ResetDrawElementList();
-	WindowElements->ResetBatchData();
-
+	WindowElements->ResetElementList();
 	WindowElementLists.Add(WindowElements);
 
 	return *WindowElements;
@@ -47,8 +46,6 @@ void FSlateDrawBuffer::RemoveUnusedWindowElement(const TArray<SWindow*>& AllWind
 		SWindow* CandidateWindow = WindowElementLists[WindowIndex]->GetPaintWindow();
 		if (!CandidateWindow || !AllWindows.Contains(CandidateWindow))
 		{
-			WindowElementLists[WindowIndex]->ResetDrawElementList();
-			WindowElementLists[WindowIndex]->ResetBatchData();
 			WindowElementLists.RemoveAtSwap(WindowIndex);
 			--WindowIndex;
 		}
@@ -57,31 +54,18 @@ void FSlateDrawBuffer::RemoveUnusedWindowElement(const TArray<SWindow*>& AllWind
 
 bool FSlateDrawBuffer::Lock()
 {
-	bool ExpectedValue = false;
-	bool bIsLock = bLocked.compare_exchange_strong(ExpectedValue, true);
-	if (bIsLock)
-	{
-		bIsLockedBySlateThread = IsInSlateThread();
-	}
-	return bIsLock;
+	return FPlatformAtomics::InterlockedCompareExchange(&Locked, 1, 0) == 0;
 }
 
 void FSlateDrawBuffer::Unlock()
 {
-	// Rendering doesn't need the batch data anymore
-	for (TSharedRef<FSlateWindowElementList>& ExistingElementList : WindowElementLists)
-	{
-		ExistingElementList->ResetBatchData();
-	}
-
-	bLocked = false;
+	FPlatformAtomics::InterlockedExchange(&Locked, 0);
 }
 
 void FSlateDrawBuffer::AddReferencedObjects(FReferenceCollector& Collector)
 {
-	// Locked buffers are the only ones that are currently referencing objects.
-	//If unlocked, the element list is not in-use and contains "to-be-cleared" data.
-	if(bLocked && !bIsLockedBySlateThread)
+	// Locked buffers are the only ones that are currently referencing objects.  If unlocked, the element list is not in-use and contains to-be-cleared data
+	if(Locked != 0)
 	{
 		for (TSharedRef<FSlateWindowElementList>& ElementList : WindowElementLists)
 		{
diff --git a/Engine/Source/Runtime/SlateCore/Public/Rendering/DrawElements.h b/Engine/Source/Runtime/SlateCore/Public/Rendering/DrawElements.h
--- a/Engine/Source/Runtime/SlateCore/Public/Rendering/DrawElements.h
+++ b/Engine/Source/Runtime/SlateCore/Public/Rendering/DrawElements.h
@@ -313,7 +313,6 @@ private:
 
 	static FVector2D GetRotationPoint( const FPaintGeometry& PaintGeometry, const TOptional<FVector2D>& UserRotationPoint, ERotationSpace RotationSpace );
 	static FSlateDrawElement& MakeBoxInternal(FSlateWindowElementList& ElementList, uint32 InLayer, const FPaintGeometry& PaintGeometry, const FSlateBrush* InBrush, ESlateDrawEffect InDrawEffects, const FLinearColor& InTint);
-
 private:
 	FSlateDataPayload* DataPayload;
 	FSlateRenderTransform RenderTransform;
@@ -534,7 +533,9 @@ public:
 		return *Payload;
 	}
 
-	/** Creates an uninitialized draw element */
+	/**
+	 * Creates an uninitialized draw element
+	 */
 	SLATECORE_API FSlateDrawElement& AddUninitialized();
 
 
@@ -618,16 +619,11 @@ public:
 	// OTHER
 	//--------------------------------------------------------------------------
 	
-	/** Remove all the elements from this draw list. */
-	UE_DEPRECATED(5.0, "ResetElementList is deprecated. Use ResetBatchData or ResetDrawElementList to reset the elements")
+	/**
+	 * Remove all the elements from this draw list.
+	 */
 	SLATECORE_API void ResetElementList();
 
-	/** Remove all the batch data. */
-	SLATECORE_API void ResetBatchData();
-
-	/** Remove all the draw elements needed to create the batch data. */
-	SLATECORE_API void ResetDrawElementList();
-
 
 	FSlateBatchData& GetBatchData() { return BatchData; }
 
@@ -641,12 +637,11 @@ private:
 	TArrayView<FSlateCachedElementData* const> GetCachedElementDataList() const { return MakeArrayView(CachedElementDataList.GetData(), CachedElementDataList.Num()); }
 
 	FSlateCachedElementData* GetCurrentCachedElementData() const { return CachedElementDataListStack.Num() ? CachedElementDataList[CachedElementDataListStack.Top()] : nullptr; }
-
 private:
 	/**
-	 * Window which owns the widgets that are being painted but not necessarily rendered to
-	 * Widgets are always rendered to the RenderTargetWindow
-	 */
+	* Window which owns the widgets that are being painted but not necessarily rendered to
+	* Widgets are always rendered to the RenderTargetWindow
+	*/
 	TWeakPtr<SWindow> WeakPaintWindow;
 	SWindow* RawPaintWindow;
 
@@ -674,7 +669,7 @@ private:
 	bool bNeedsDeferredResolve;
 	TArray<int32> ResolveToDeferredIndex;
 
-	//~ Begin Fast Path
+	// Begin Fast Path
 
 	/** State of the current widget that is adding draw elements */
 	struct FWidgetDrawElementState
@@ -697,7 +692,7 @@ private:
 	TArray<FWidgetDrawElementState, TInlineAllocator<50>> WidgetDrawStack;
 	TArray<FSlateCachedElementData*, TInlineAllocator<4>> CachedElementDataList;
 	TArray<int32, TInlineAllocator<4>> CachedElementDataListStack;
-	//~ End Fast Path
+	// End Fast Path
 
 	/** Store the size of the window being used to paint */
 	FVector2D WindowSize;
diff --git a/Engine/Source/Runtime/SlateCore/Public/Rendering/SlateDrawBuffer.h b/Engine/Source/Runtime/SlateCore/Public/Rendering/SlateDrawBuffer.h
--- a/Engine/Source/Runtime/SlateCore/Public/Rendering/SlateDrawBuffer.h
+++ b/Engine/Source/Runtime/SlateCore/Public/Rendering/SlateDrawBuffer.h
@@ -5,8 +5,6 @@
 #include "CoreMinimal.h"
 #include "UObject/GCObject.h"
 
-#include <atomic>
-
 class FSlateWindowElementList;
 class SWindow;
 
@@ -19,8 +17,7 @@ public:
 
 	/** Default constructor. */
 	explicit FSlateDrawBuffer()
-		: bLocked(false)
-		, bIsLockedBySlateThread(false)
+		: Locked(0)
 		, ResourceVersion(0)
 	{ }
 
@@ -81,8 +78,8 @@ protected:
 	// that we restore if they're requested again.
 	TArray< TSharedRef<FSlateWindowElementList> > WindowElementListsPool;
 
-	std::atomic<bool> bLocked;
-	bool bIsLockedBySlateThread;
+	// 1 if this buffer is locked, 0 otherwise.
+	volatile int32 Locked;
 
 	// Last recorded version from the render. The WindowElementListsPool is emptied when this changes.
 	uint32 ResourceVersion;
diff --git a/Engine/Source/Runtime/SlateRHIRenderer/Private/SlateRHIRenderer.cpp b/Engine/Source/Runtime/SlateRHIRenderer/Private/SlateRHIRenderer.cpp
--- a/Engine/Source/Runtime/SlateRHIRenderer/Private/SlateRHIRenderer.cpp
+++ b/Engine/Source/Runtime/SlateRHIRenderer/Private/SlateRHIRenderer.cpp
@@ -91,9 +91,8 @@ TAutoConsoleVariable<int32> CVarShowSlateBatching(
 struct FSlateDrawWindowCommandParams
 {
 	FSlateRHIRenderer* Renderer;
+	FSlateWindowElementList* WindowElementList;
 	SWindow* Window;
-	FSlateBatchData* BatchData;
-	FVector2D WindowSize;
 #if WANTS_DRAW_MESH_EVENTS
 	FString WindowTitle;
 #endif
@@ -677,7 +676,7 @@ int32 SlateWireFrame = 0;
 static FAutoConsoleVariableRef CVarSlateWireframe(TEXT("Slate.ShowWireFrame"), SlateWireFrame, TEXT(""), ECVF_Default);
 
 /** Draws windows from a FSlateDrawBuffer on the render thread */
-void FSlateRHIRenderer::DrawWindow_RenderThread(FRHICommandListImmediate& RHICmdList, FViewportInfo& ViewportInfo, struct FSlateDrawWindowCommandParams& DrawCommandParams)
+void FSlateRHIRenderer::DrawWindow_RenderThread(FRHICommandListImmediate& RHICmdList, FViewportInfo& ViewportInfo, FSlateWindowElementList& WindowElementList, const struct FSlateDrawWindowCommandParams& DrawCommandParams)
 {
 	LLM_SCOPE(ELLMTag::SceneRender);
 	
@@ -731,7 +730,7 @@ void FSlateRHIRenderer::DrawWindow_RenderThread(FRHICommandListImmediate& RHICmd
 		bool bRenderedStereo = false;
 		if (CVarDrawToVRRenderTarget->GetInt() == 0 && GEngine && IsValidRef(ViewportInfo.GetRenderTargetTexture()) && GEngine->StereoRenderingDevice.IsValid())
 		{
-			const FVector2D WindowSize = DrawCommandParams.WindowSize;
+			const FVector2D WindowSize = WindowElementList.GetWindowSize();
 			GEngine->StereoRenderingDevice->RenderTexture_RenderThread(RHICmdList, RHICmdList.GetViewportBackBuffer(ViewportInfo.ViewportRHI), ViewportInfo.GetRenderTargetTexture(), WindowSize);
 			bRenderedStereo = true;
 		}
@@ -741,8 +740,7 @@ void FSlateRHIRenderer::DrawWindow_RenderThread(FRHICommandListImmediate& RHICmd
 			SCOPE_CYCLE_COUNTER(STAT_SlateRenderingRTTime);
 			CSV_SCOPED_TIMING_STAT_EXCLUSIVE(Slate);
 
-			check(DrawCommandParams.BatchData);
-			FSlateBatchData& BatchData = *DrawCommandParams.BatchData;
+			FSlateBatchData& BatchData = WindowElementList.GetBatchData();
 
 			// Update the vertex and index buffer	
 			RenderingPolicy->BuildRenderingBuffers(RHICmdList, BatchData);
@@ -1095,7 +1093,7 @@ void FSlateRHIRenderer::DrawWindow_RenderThread(FRHICommandListImmediate& RHICmd
 #endif
 			if (!bRenderedStereo && GEngine && IsValidRef(ViewportInfo.GetRenderTargetTexture()) && GEngine->StereoRenderingDevice.IsValid())
 			{
-				const FVector2D WindowSize = DrawCommandParams.WindowSize;
+				const FVector2D WindowSize = WindowElementList.GetWindowSize();
 				GEngine->StereoRenderingDevice->RenderTexture_RenderThread(RHICmdList, RHICmdList.GetViewportBackBuffer(ViewportInfo.ViewportRHI), ViewportInfo.GetRenderTargetTexture(), WindowSize);
 			}
 			RHICmdList.Transition(FRHITransitionInfo(BackBuffer, ERHIAccess::Unknown, ERHIAccess::SRVGraphics));
@@ -1313,8 +1311,7 @@ void FSlateRHIRenderer::DrawWindows_Private(FSlateDrawBuffer& WindowDrawBuffer)
 					FSlateDrawWindowCommandParams Params;
 
 					Params.Renderer = this;
-					Params.WindowSize = ElementList.GetWindowSize();
-					Params.BatchData = &ElementList.GetBatchData();
+					Params.WindowElementList = &ElementList;
 					Params.Window = Window;
 #if WANTS_DRAW_MESH_EVENTS
 					Params.WindowTitle = Window->GetTitle().ToString();
@@ -1334,9 +1331,9 @@ void FSlateRHIRenderer::DrawWindows_Private(FSlateDrawBuffer& WindowDrawBuffer)
 					if (GIsClient && !IsRunningCommandlet() && !GUsingNullRHI)
 					{
 						ENQUEUE_RENDER_COMMAND(SlateDrawWindowsCommand)(
-							[Params, ViewInfo](FRHICommandListImmediate& RHICmdList) mutable
+							[Params, ViewInfo](FRHICommandListImmediate& RHICmdList)
 							{
-								Params.Renderer->DrawWindow_RenderThread(RHICmdList, *ViewInfo, Params);
+								Params.Renderer->DrawWindow_RenderThread(RHICmdList, *ViewInfo, *Params.WindowElementList, Params);
 							}
 						);
 					}
diff --git a/Engine/Source/Runtime/SlateRHIRenderer/Private/SlateRHIRenderer.h b/Engine/Source/Runtime/SlateRHIRenderer/Private/SlateRHIRenderer.h
--- a/Engine/Source/Runtime/SlateRHIRenderer/Private/SlateRHIRenderer.h
+++ b/Engine/Source/Runtime/SlateRHIRenderer/Private/SlateRHIRenderer.h
@@ -214,7 +214,7 @@ public:
 	virtual void AddWidgetRendererUpdate(const struct FRenderThreadUpdateContext& Context, bool bDeferredRenderTargetUpdate) override;
 
 	/** Draws windows from a FSlateDrawBuffer on the render thread */
-	void DrawWindow_RenderThread(FRHICommandListImmediate& RHICmdList, FViewportInfo& ViewportInfo, struct FSlateDrawWindowCommandParams& DrawParams);
+	void DrawWindow_RenderThread(FRHICommandListImmediate& RHICmdList, FViewportInfo& ViewportInfo, FSlateWindowElementList& WindowElementList, const struct FSlateDrawWindowCommandParams& DrawParams);
 
 	/**
 	 * Reloads texture resources from disk                   
diff --git a/Engine/Source/Runtime/UMG/Private/Slate/WidgetRenderer.cpp b/Engine/Source/Runtime/UMG/Private/Slate/WidgetRenderer.cpp
--- a/Engine/Source/Runtime/UMG/Private/Slate/WidgetRenderer.cpp
+++ b/Engine/Source/Runtime/UMG/Private/Slate/WidgetRenderer.cpp
@@ -342,7 +342,7 @@ bool FWidgetRenderer::DrawInvalidationRoot(TSharedRef<SVirtualWindow>& VirtualWi
 		}
 		else
 		{
-			WindowElementList.ResetDrawElementList();
+			WindowElementList.ResetElementList();
 			DrawBuffer.Unlock();
 		}
 	}
