Removed a few high-frequency/low-value Slate RT stats to improve profiler readability and stability

diff --git a/Engine/Source/Runtime/SlateRHIRenderer/Private/SlateRHIRenderingPolicy.cpp b/Engine/Source/Runtime/SlateRHIRenderer/Private/SlateRHIRenderingPolicy.cpp
--- a/Engine/Source/Runtime/SlateRHIRenderer/Private/SlateRHIRenderingPolicy.cpp
+++ b/Engine/Source/Runtime/SlateRHIRenderer/Private/SlateRHIRenderingPolicy.cpp
@@ -32,16 +32,10 @@ extern void UpdateNoiseTextureParameters(FViewUniformShaderParameters& ViewUnifo
 DECLARE_CYCLE_STAT(TEXT("Update Buffers RT"), STAT_SlateUpdateBufferRTTime, STATGROUP_Slate);
 DECLARE_CYCLE_STAT(TEXT("Update Buffers RT"), STAT_SlateUpdateBufferRTTimeLambda, STATGROUP_Slate);
 
-
 DECLARE_DWORD_COUNTER_STAT(TEXT("Num Layers"), STAT_SlateNumLayers, STATGROUP_Slate);
 DECLARE_DWORD_COUNTER_STAT(TEXT("Num Batches"), STAT_SlateNumBatches, STATGROUP_Slate);
 DECLARE_DWORD_COUNTER_STAT(TEXT("Num Vertices"), STAT_SlateVertexCount, STATGROUP_Slate);
 
-DECLARE_CYCLE_STAT(TEXT("Slate RT: Texture Draw Call"), STAT_SlateRTTextureDrawCall, STATGROUP_Slate);
-DECLARE_CYCLE_STAT(TEXT("Slate RT: Material Draw Call"), STAT_SlateRTMaterialDrawCall, STATGROUP_Slate);
-DECLARE_CYCLE_STAT(TEXT("Slate RT: Scissor Draw Call"), STAT_SlateRTStencilDrawCall, STATGROUP_Slate);
-DECLARE_CYCLE_STAT(TEXT("Slate RT: Custom Draw"), STAT_SlateRTCustomDraw, STATGROUP_Slate);
-
 DECLARE_DWORD_COUNTER_STAT(TEXT("Clips (Scissor)"), STAT_SlateScissorClips, STATGROUP_Slate);
 DECLARE_DWORD_COUNTER_STAT(TEXT("Clips (Stencil)"), STAT_SlateStencilClips, STATGROUP_Slate);
 
@@ -463,8 +457,6 @@ static bool UpdateScissorRect(
 					{
 						const FSlateClippingZone& MaskQuad = StencilQuads[0];
 
-						SCOPE_CYCLE_COUNTER(STAT_SlateRTStencilDrawCall);
-
 						VertexShader->SetMaskRect(RHICmdList, FVector2D(MaskQuad.TopLeft), FVector2D(MaskQuad.TopRight), FVector2D(MaskQuad.BottomLeft), FVector2D(MaskQuad.BottomRight));
 
 						RHICmdList.SetStreamSource(0, StencilVertexBuffer.VertexBufferRHI, 0);
@@ -508,8 +500,6 @@ static bool UpdateScissorRect(
 				{
 					const FSlateClippingZone& MaskQuad = StencilQuads[MaskIndex];
 
-					SCOPE_CYCLE_COUNTER(STAT_SlateRTStencilDrawCall);
-					
 					VertexShader->SetMaskRect(RHICmdList, FVector2D(MaskQuad.TopLeft), FVector2D(MaskQuad.TopRight), FVector2D(MaskQuad.BottomLeft), FVector2D(MaskQuad.BottomRight));
 
 					RHICmdList.SetStreamSource(0, StencilVertexBuffer.VertexBufferRHI, 0);
@@ -999,8 +989,6 @@ void FSlateRHIRenderingPolicy::DrawElements(
 				}
 
 				{
-					QUICK_SCOPE_CYCLE_COUNTER(Slate_SetTextureShaderParams);
-
 					GlobalVertexShader->SetViewProjection(RHICmdList, FMatrix44f(ViewProjection));
 					GlobalVertexShader->SetVerticalAxisMultiplier(RHICmdList, bSwitchVerticalAxis ? -1.0f : 1.0f);
 
@@ -1020,8 +1008,6 @@ void FSlateRHIRenderingPolicy::DrawElements(
 				}
 
 				{
-					SCOPE_CYCLE_COUNTER(STAT_SlateRTTextureDrawCall);
-
 					// for RHIs that can't handle VertexOffset, we need to offset the stream source each time
 					RHICmdList.SetStreamSource(0, VertexBufferPtr->VertexBufferRHI, RenderBatch.VertexOffset * sizeof(FSlateVertex));
 					RHICmdList.DrawIndexedPrimitive(IndexBufferPtr->IndexBufferRHI, 0, 0, RenderBatch.NumVertices, RenderBatch.IndexOffset, PrimitiveCount, RenderBatch.InstanceCount);
@@ -1142,7 +1128,6 @@ void FSlateRHIRenderingPolicy::DrawElements(
 							SetGraphicsPipelineState(RHICmdList, GraphicsPSOInit, StencilRef);
 
 							{
-								QUICK_SCOPE_CYCLE_COUNTER(Slate_SetMaterialShaderParams);
 								VertexShader->SetViewProjection(RHICmdList, FMatrix44f(ViewProjection));
 								VertexShader->SetVerticalAxisMultiplier(RHICmdList, bSwitchVerticalAxis ? -1.0f : 1.0f);
 								VertexShader->SetMaterialShaderParameters(RHICmdList, ActiveSceneView, MaterialRenderProxy, EffectiveMaterial);
@@ -1165,7 +1150,6 @@ void FSlateRHIRenderingPolicy::DrawElements(
 						}
 
 						{
-							SCOPE_CYCLE_COUNTER(STAT_SlateRTMaterialDrawCall);
 							if (bUseInstancing)
 							{
 								uint32 InstanceCount = RenderBatch.InstanceCount;
@@ -1343,7 +1327,6 @@ ETextureSamplerFilter FSlateRHIRenderingPolicy::GetSamplerFilter(const UTexture*
 TShaderRef<FSlateElementPS> FSlateRHIRenderingPolicy::GetTexturePixelShader( FGlobalShaderMap* ShaderMap, ESlateShader ShaderType, ESlateDrawEffect DrawEffects, bool bIsVirtualTexture )
 {
 	TShaderRef<FSlateElementPS> PixelShader;
-	QUICK_SCOPE_CYCLE_COUNTER(STAT_Slate_GetTexturePixelShader);
 
 #if WITH_SLATE_VISUALIZERS
 	if ( CVarShowSlateOverdraw.GetValueOnRenderThread() != 0 )
