support for HDR change on process resume

diff --git a/Engine/Source/Runtime/D3D12RHI/Private/D3D12Viewport.cpp b/Engine/Source/Runtime/D3D12RHI/Private/D3D12Viewport.cpp
--- a/Engine/Source/Runtime/D3D12RHI/Private/D3D12Viewport.cpp
+++ b/Engine/Source/Runtime/D3D12RHI/Private/D3D12Viewport.cpp
@@ -913,6 +913,11 @@ bool FD3D12Viewport::CheckHDRSupport()
 	return GRHISupportsHDROutput && IsHDREnabled();
 }
 
+EPixelFormat GetDefaultBackBufferPixelFormat()
+{
+	static const auto CVarDefaultBackBufferPixelFormat = IConsoleManager::Get().FindTConsoleVariableDataInt(TEXT("r.DefaultBackBufferPixelFormat"));
+	return EDefaultBackBufferPixelFormat::Convert2PixelFormat(EDefaultBackBufferPixelFormat::FromInt(CVarDefaultBackBufferPixelFormat->GetValueOnGameThread()));
+}
 
 /*==============================================================================
  *	The following RHI functions must be called from the main thread.
@@ -923,8 +928,7 @@ FViewportRHIRef FD3D12DynamicRHI::RHICreateViewport(void* WindowHandle, uint32 S
 
 	if (PreferredPixelFormat == EPixelFormat::PF_Unknown)
 	{
-		static const auto CVarDefaultBackBufferPixelFormat = IConsoleManager::Get().FindTConsoleVariableDataInt(TEXT("r.DefaultBackBufferPixelFormat"));
-		PreferredPixelFormat = EDefaultBackBufferPixelFormat::Convert2PixelFormat(EDefaultBackBufferPixelFormat::FromInt(CVarDefaultBackBufferPixelFormat->GetValueOnGameThread()));
+		PreferredPixelFormat = GetDefaultBackBufferPixelFormat();
 	}
 
 	FD3D12Viewport* RenderingViewport = new FD3D12Viewport(&GetAdapter(), (HWND)WindowHandle, SizeX, SizeY, bIsFullscreen, PreferredPixelFormat);
diff --git a/Engine/Source/Runtime/SlateRHIRenderer/Private/SlateRHIRenderer.cpp b/Engine/Source/Runtime/SlateRHIRenderer/Private/SlateRHIRenderer.cpp
--- a/Engine/Source/Runtime/SlateRHIRenderer/Private/SlateRHIRenderer.cpp
+++ b/Engine/Source/Runtime/SlateRHIRenderer/Private/SlateRHIRenderer.cpp
@@ -376,7 +376,9 @@ void FSlateRHIRenderer::CreateViewport(const TSharedRef<SWindow> Window)
 
 		// SDR format holds the requested format in non HDR mode
 		NewInfo->SDRPixelFormat = NewInfo->PixelFormat;
-		if (IsHDREnabled())
+		bool HDREnabled = IsHDREnabled();
+		NewInfo->bSceneHDREnabled = HDREnabled;
+		if (HDREnabled)
 		{
 			NewInfo->PixelFormat = GRHIHDRDisplayOutputFormat;
 		}
@@ -414,11 +416,16 @@ void FSlateRHIRenderer::ConditionalResizeViewport(FViewportInfo* ViewInfo, uint3
 		((ViewInfo->PixelFormat == GRHIHDRDisplayOutputFormat) != bHDREnabled)	// HDR toggled
 #if PLATFORM_WINDOWS
 		|| ((IsRHIDeviceNVIDIA() || IsRHIDeviceAMD()) &&						// Vendor-specific mastering data updates
-		((bHDREnabled && ViewInfo->HDRColorGamut != HDRColorGamut)			// Color gamut changed
+		((bHDREnabled && ViewInfo->HDRColorGamut != HDRColorGamut)				// Color gamut changed
 			|| (bHDREnabled && ViewInfo->HDROutputDevice != HDROutputDevice)))	// Output device changed
 #endif
 		);
 
+	if (ViewInfo && (ViewInfo->bSceneHDREnabled != bHDREnabled))
+	{
+		bHDRStale = true;
+	}
+
 	if (IsInGameThread() && !IsInSlateThread() && ViewInfo && (bHDRStale || ViewInfo->Height != Height || ViewInfo->Width != Width || ViewInfo->bFullscreen != bFullscreen || !IsValidRef(ViewInfo->ViewportRHI)))
 	{
 		// The viewport size we have doesn't match the requested size of the viewport.
@@ -459,6 +466,7 @@ void FSlateRHIRenderer::ConditionalResizeViewport(FViewportInfo* ViewInfo, uint3
 		ViewInfo->PixelFormat = bHDREnabled ? GRHIHDRDisplayOutputFormat : ViewInfo->SDRPixelFormat;
 		ViewInfo->HDRColorGamut = HDRColorGamut;
 		ViewInfo->HDROutputDevice = HDROutputDevice;
+		ViewInfo->bSceneHDREnabled = bHDREnabled;
 
 		PreResizeBackBufferDelegate.Broadcast(&ViewInfo->ViewportRHI);
 		if (IsValidRef(ViewInfo->ViewportRHI))
diff --git a/Engine/Source/Runtime/SlateRHIRenderer/Private/SlateRHIRenderer.h b/Engine/Source/Runtime/SlateRHIRenderer/Private/SlateRHIRenderer.h
--- a/Engine/Source/Runtime/SlateRHIRenderer/Private/SlateRHIRenderer.h
+++ b/Engine/Source/Runtime/SlateRHIRenderer/Private/SlateRHIRenderer.h
@@ -100,7 +100,10 @@ struct FViewportInfo : public FRenderResource
 	
 	/** Whether is in a HDR Color Space */
 	bool bHDREnabled;
-	
+
+	/** Whether the scene output is in HDR. bHDREnabled only affects Slate rendering */
+	bool bSceneHDREnabled;
+
 	/** FRenderResource interface */
 	virtual void InitRHI() override;
 	virtual void ReleaseRHI() override;
@@ -119,7 +122,8 @@ struct FViewportInfo : public FRenderResource
 			PixelFormat(EPixelFormat::PF_Unknown),
 			SDRPixelFormat(EPixelFormat::PF_Unknown),
 			RTProvider(nullptr),
-			bHDREnabled(false)
+			bHDREnabled(false),
+			bSceneHDREnabled(false)
 	{
 	}
 
