Fix Interface Flickering due to uninitialized variable during the ColorBlind pass. Because of that alphablend was enabled.

#rnx

diff --git a/Engine/Source/Runtime/SlateRHIRenderer/Private/SlateRHIRenderingPolicy.cpp b/Engine/Source/Runtime/SlateRHIRenderer/Private/SlateRHIRenderingPolicy.cpp
--- a/Engine/Source/Runtime/SlateRHIRenderer/Private/SlateRHIRenderingPolicy.cpp
+++ b/Engine/Source/Runtime/SlateRHIRenderer/Private/SlateRHIRenderingPolicy.cpp
@@ -1263,6 +1263,7 @@ void FSlateRHIRenderingPolicy::DrawElements(
 		RectParams.SourceRect = FSlateRect(0, 0, BackBuffer.GetSizeXY().X, BackBuffer.GetSizeXY().Y);
 		RectParams.DestRect = FSlateRect(0, 0, BackBuffer.GetSizeXY().X, BackBuffer.GetSizeXY().Y);
 		RectParams.SourceTextureSize = BackBuffer.GetSizeXY();
+		RectParams.CornerRadius = FVector4(0, 0, 0, 0);
 
 		PostProcessor->ColorDeficiency(RHICmdList, RendererModule, RectParams);
 
