bugfixes for HDR in editor:
conversion to scRGB was happening twice in CompositeUIPixelShader: once during the LUT, and once during ST2084ToScRGB. We now rely only on the LUT to put us to the correct color space, just like we do during tonemapping
Make sure to use rec709 instead of rec2020 in editor, since this is the color space that is used by scRGB (PF_FloatRGBA backbuffer format)
Fixed UI composition to properly composite HDR elements (like the viewport) with the rest of the UI: we draw a black, 0-alpha rectangle as a proxy in the SDR pass for HDR elements, and we draw the HDR elements in their own rendertargets

diff --git a/Engine/Shaders/Private/CompositeUIPixelShader.usf b/Engine/Shaders/Private/CompositeUIPixelShader.usf
--- a/Engine/Shaders/Private/CompositeUIPixelShader.usf
+++ b/Engine/Shaders/Private/CompositeUIPixelShader.usf
@@ -9,6 +9,9 @@ CompositeUIPixelShader.usf: Filter pixel shader source.
 #define LUT_SCALING	1.05f
 #define LUT_SIZE	32.f
 
+// set this to 1 if you want to restore the legacy behaviour (ie apply ACES to the UI). A value of 0 keeps a consistent look of the UI between SDR and HDR mode
+#define USE_LUT_FOR_UI 0
+
 ///////////////////////////////////////////////////////////////////////////////////////
 uint OutputDevice;
 uint OutputGamut;
@@ -30,14 +33,23 @@ void GenerateLUTPS(
 	{		
 		// ACES 1000nit transform with PQ/2084 encoding, user specified gamut 
 		float3 ODTColor = ACESOutputTransforms1000(LinearColor);
-		ODTColor = mul(AP1_2_Output, ODTColor);
+		// do not apply scRGB here because LUT is not 16F in this format, and scRGB can output negative values
+		if (OutputDevice != TONEMAPPER_OUTPUT_ACES1000nitScRGB)
+		{
+			ODTColor = mul(AP1_2_Output, ODTColor);
+		}
 		OutColor.rgb = LinearToST2084(ODTColor);
 	}	
 	else if( OutputDevice == TONEMAPPER_OUTPUT_ACES2000nitST2084 || OutputDevice == TONEMAPPER_OUTPUT_ACES2000nitScRGB)
 	{		
 		// ACES 2000nit transform with PQ/2084 encoding, user specified gamut 
 		float3 ODTColor = ACESOutputTransforms2000(LinearColor);
-		ODTColor = mul(AP1_2_Output, ODTColor);
+		// do not apply scRGB here because LUT is not 16F in this format, and scRGB can output negative values
+		// 
+		if (OutputDevice != TONEMAPPER_OUTPUT_ACES2000nitScRGB)
+		{
+			ODTColor = mul(AP1_2_Output, ODTColor);
+		}
 		OutColor.rgb = LinearToST2084(ODTColor);
 	}
 	else
@@ -65,6 +77,22 @@ SamplerState ColorSpaceLUTSampler;
 
 float UILevel;
 
+// Invert the encoding applied by the ApplyGammaCorrection function.
+float3 LinearizeColor(float3 EncodedColor)
+{
+	#if MAC
+		// Note, MacOSX native output is raw gamma 2.2 not sRGB!
+		return pow(EncodedColor, 2.2);
+	#else
+		#if USE_709
+			// Didn't profile yet if the branching version would be faster (different linear segment).
+			return Rec709ToLinear(EncodedColor);
+		#else
+			return sRGBToLinear(EncodedColor);
+		#endif
+	#endif
+}
+
 float3 ColorLookupTable(float3 LinearColor)
 {
 	float3 LUTEncodedColor = LinearToST2084(LinearColor);
@@ -107,18 +135,36 @@ void Main(
 		discard;
 	}
 #endif
+
+	// In order to get consistent blending between HDR PQ & ScRGB, we convert UI color in both cases to rec2020, perform the blending in this color space, and
+	// move back to the output device space. Rec2020 is chosen over ScRGB because this is the way blending is currently being done in engine mode, and contains only positive values
+	const float3x3 sRGB_2_Rec2020 = mul(XYZ_2_Rec2020_MAT, sRGB_2_XYZ_MAT);
+	static const float PaperWhiteInNits = 80.0f; // used to convert between ScRGB and nits
+	static const float UIBrightnessInNits = 300.0f; // todo stop hardcoding this value and try to fetch SDR/HDR brightness balance instead (automatically or though calibration step)
+
 	float4 UIColor = Texture2DSample(UITexture, UISampler, Input.UV);
-	float3 SceneColor = Texture2DSample(SceneTexture, SceneSampler, Input.UV).xyz;
+	{
+		UIColor.xyz = LinearizeColor(UIColor.xyz);
+#if USE_LUT_FOR_UI
+		UIColor.xyz = ColorLookupTable(UIColor.xyz);
+		// PQ -> Linear
+#if !SCRGB_ENCODING
+		UIColor.xyz = ST2084ToLinear(UIColor.xyz);
+#else
+		UIColor.xyz = ST2084ToScRGB(UIColor.xyz, OutputDevice);
+		UIColor.xyz = mul(sRGB_2_Rec2020, UIColor.xyz) * UIBrightnessInNits;
+#endif
 
-	// sRGB -> RRT -> ODT
-	UIColor.xyz = ColorLookupTable(UIColor.xyz);
+#else // #if USE_LUT_FOR_UI
+		UIColor.xyz = mul(sRGB_2_Rec2020, UIColor.xyz) * UIBrightnessInNits;
+#endif // #if USE_LUT_FOR_UI
+	}
 
-	// PQ -> Linear
+	float3 SceneColor = Texture2DSample(SceneTexture, SceneSampler, Input.UV).xyz;
 #if !SCRGB_ENCODING
 	SceneColor = ST2084ToLinear(SceneColor);
-	UIColor.xyz = ST2084ToLinear(UIColor.xyz);
 #else
-	UIColor.xyz = ST2084ToScRGB(UIColor.xyz, OutputDevice);
+	SceneColor.xyz = mul(sRGB_2_Rec2020, SceneColor.xyz) * PaperWhiteInNits;
 #endif
 
 	BRANCH
@@ -129,7 +175,9 @@ void Main(
 
 		// Tonemap HDR under transparent UI with a simple Rienhard to the max luminance of the UI
 		// This prevents HDR bleed through destroying UI legibility
-		float Luminance = dot(SceneColor, half3(0.299f, 0.587f, 0.114f)) / UILevel;
+		// Rec2020 coefficients to compute luminance
+		float KR = 0.2627, KG = 0.678, KB = 0.0593;
+		float Luminance = dot(SceneColor, half3(KR, KG, KB)) / UILevel;
 		float OutL = 1.f / (Luminance + 1.f);
 
 		// Ease out remapping to avoid hard transitions where UI is near zero opacity
@@ -143,5 +191,8 @@ void Main(
 #if !SCRGB_ENCODING
 	// Linear -> PQ
 	OutColor.xyz = LinearToST2084(OutColor.xyz);
+#else
+	const float3x3 Rec2020_2_sRGB = mul(XYZ_2_sRGB_MAT, Rec2020_2_XYZ_MAT);
+	OutColor.xyz = mul(Rec2020_2_sRGB, OutColor.xyz / PaperWhiteInNits);
 #endif
 }
\ No newline at end of file
diff --git a/Engine/Shaders/Private/GammaCorrectionCommon.ush b/Engine/Shaders/Private/GammaCorrectionCommon.ush
--- a/Engine/Shaders/Private/GammaCorrectionCommon.ush
+++ b/Engine/Shaders/Private/GammaCorrectionCommon.ush
@@ -8,6 +8,12 @@ half3 LinearTo709Branchless(half3 lin)
 	return min(lin * 4.5, pow(max(lin, 0.018), 0.45) * 1.099 - 0.099);
 }
 
+half3 Rec709ToLinear(half3 Color)
+{
+	Color = max(6.10352e-5, Color); // minimum positive non-denormal (fixes black problem on DX11 AMD and NV)
+	return Color > 0.081 ? pow((Color + 0.099) / 1.099, 1.0 / 0.45) : Color / 4.5;
+}
+
 half3 LinearToSrgbBranchless(half3 lin) 
 {
 	lin = max(6.10352e-5, lin); // minimum positive non-denormal (fixes black problem on DX11 AMD and NV)
diff --git a/Engine/Shaders/Private/TonemapCommon.ush b/Engine/Shaders/Private/TonemapCommon.ush
--- a/Engine/Shaders/Private/TonemapCommon.ush
+++ b/Engine/Shaders/Private/TonemapCommon.ush
@@ -341,7 +341,7 @@ float3x3 OuputInverseGamutMappingMatrix( uint OutputGamut )
 float3 ST2084ToScRGB(float3 Color, uint OutputDevice)
 {
 	// Nvidia HDR encoding - Remove PQ, convert to linear scRGB
-	const float3x3 AP1_2_sRGB = mul(XYZ_2_sRGB_MAT, AP1_2_XYZ_MAT);
+	const float3x3 AP1_2_sRGB = mul(XYZ_2_sRGB_MAT, mul(D60_2_D65_CAT, AP1_2_XYZ_MAT));
 	const float WhitePoint = 80.f;
 
 	// 1000.f nit display
diff --git a/Engine/Source/Runtime/Core/Public/Windows/WindowsPlatformMisc.h b/Engine/Source/Runtime/Core/Public/Windows/WindowsPlatformMisc.h
--- a/Engine/Source/Runtime/Core/Public/Windows/WindowsPlatformMisc.h
+++ b/Engine/Source/Runtime/Core/Public/Windows/WindowsPlatformMisc.h
@@ -259,9 +259,18 @@ struct CORE_API FWindowsPlatformMisc
 	{
 		if (DeviceId == 0x1002 /*AMD*/ || DeviceId == 0x10DE /*NVIDIA*/)
 		{
-			// ScRGB, 1000 or 2000 nits, Rec2020
+			// needs to match GRHIHDRDisplayOutputFormat chosen in FD3D12DynamicRHI::Init
+#if WITH_EDITOR
+		// ScRGB, 1000 or 2000 nits
 			OutputDevice = (DisplayNitLevel == 1000) ? (int32)EDisplayOutputFormat::HDR_ACES_1000nit_ScRGB : (int32)EDisplayOutputFormat::HDR_ACES_2000nit_ScRGB;
+			// Rec709
+			ColorGamut = 0;
+#else
+		// ST-2084, 1000 or 2000 nits
+			OutputDevice = (DisplayNitLevel == 1000) ? (int32)EDisplayOutputFormat::HDR_ACES_1000nit_ST2084 : (int32)EDisplayOutputFormat::HDR_ACES_2000nit_ST2084;
+			// Rec2020
 			ColorGamut = 2;
+#endif
 		}
 	}
 
diff --git a/Engine/Source/Runtime/Engine/Private/Slate/SceneViewport.cpp b/Engine/Source/Runtime/Engine/Private/Slate/SceneViewport.cpp
--- a/Engine/Source/Runtime/Engine/Private/Slate/SceneViewport.cpp
+++ b/Engine/Source/Runtime/Engine/Private/Slate/SceneViewport.cpp
@@ -49,6 +49,7 @@ FSceneViewport::FSceneViewport( FViewportClient* InViewportClient, TSharedPtr<SV
 	, bForceViewportSize(false)
 	, bPlayInEditorIsSimulate( false )
 	, bCursorHiddenDueToCapture( false )
+	, bHDRViewport(false)
 	, MousePosBeforeHiddenDueToCapture( -1, -1 )
 	, RTTSize( 0, 0 )
 	, NumBufferedFrames(1)
@@ -360,8 +361,10 @@ void FSceneViewport::OnDrawViewport( const FGeometry& AllottedGeometry, const FS
 	if (!bForceViewportSize)
 	{
 		FIntPoint DrawSize = FIntPoint( FMath::RoundToInt( AllottedGeometry.GetDrawSize().X ), FMath::RoundToInt( AllottedGeometry.GetDrawSize().Y ) );
-		if( GetSizeXY() != DrawSize )
+		bool bIsHDREnabled = IsHDREnabled();
+		if( GetSizeXY() != DrawSize || bIsHDREnabled != bHDRViewport)
 		{
+			bHDRViewport = bIsHDREnabled;
 			TSharedPtr<SWindow> Window = FSlateApplication::Get().FindWidgetWindow( ViewportWidget.Pin().ToSharedRef() );
 			if ( Window.IsValid() )
 			{
@@ -2004,14 +2007,10 @@ void FSceneViewport::InitDynamicRHI()
 		EPixelFormat SceneTargetFormat = EDefaultBackBufferPixelFormat::Convert2PixelFormat(EDefaultBackBufferPixelFormat::FromInt(CVarDefaultBackBufferPixelFormat->GetValueOnRenderThread()));
 		SceneTargetFormat = RHIPreferredPixelFormatHint(SceneTargetFormat);
 	
-#if WITH_EDITOR
-		// HDR Editor needs to be in float format if running with HDR
-		static auto CVarHDREnable = IConsoleManager::Get().FindConsoleVariable(TEXT("Editor.HDRSupport"));
-		if(CVarHDREnable && (CVarHDREnable->GetInt() != 0))
+		if (bHDRViewport)
 		{
-			SceneTargetFormat = PF_FloatRGBA;
+			SceneTargetFormat = GRHIHDRDisplayOutputFormat;
 		}
-#endif
 
 		FTextureRHIRef BufferedRTRHI;
 		FTextureRHIRef BufferedSRVRHI;
diff --git a/Engine/Source/Runtime/Engine/Public/Slate/SceneViewport.h b/Engine/Source/Runtime/Engine/Public/Slate/SceneViewport.h
--- a/Engine/Source/Runtime/Engine/Public/Slate/SceneViewport.h
+++ b/Engine/Source/Runtime/Engine/Public/Slate/SceneViewport.h
@@ -421,6 +421,8 @@ private:
 	bool bPlayInEditorIsSimulate;
 	/** Whether or not the cursor is hidden when the viewport captures the mouse */
 	bool bCursorHiddenDueToCapture;
+	/** Whether or not the viewport is in HDR */
+	bool bHDRViewport;
 	/** Position the cursor was at when we hid it due to capture, so we can put it back afterwards */
 	FIntPoint MousePosBeforeHiddenDueToCapture;
 	/** Dimensions of RenderTarget texture. */
diff --git a/Engine/Source/Runtime/SlateCore/Private/Rendering/DrawElements.cpp b/Engine/Source/Runtime/SlateCore/Private/Rendering/DrawElements.cpp
--- a/Engine/Source/Runtime/SlateCore/Private/Rendering/DrawElements.cpp
+++ b/Engine/Source/Runtime/SlateCore/Private/Rendering/DrawElements.cpp
@@ -967,6 +967,7 @@ void FSlateWindowElementList::ResetElementList()
 	DeferredPaintList.Reset();
 
 	BatchData.ResetData();
+	BatchDataHDR.ResetData();
 
 	ClippingManager.ResetClippingState();
 
diff --git a/Engine/Source/Runtime/SlateCore/Private/Rendering/ElementBatcher.cpp b/Engine/Source/Runtime/SlateCore/Private/Rendering/ElementBatcher.cpp
--- a/Engine/Source/Runtime/SlateCore/Private/Rendering/ElementBatcher.cpp
+++ b/Engine/Source/Runtime/SlateCore/Private/Rendering/ElementBatcher.cpp
@@ -256,7 +256,8 @@ void FSlateBatchData::MergeRenderBatches()
 
 FSlateElementBatcher::FSlateElementBatcher( TSharedRef<FSlateRenderingPolicy> InRenderingPolicy )
 	: BatchData( nullptr )
-	, CurrentCachedElementList( nullptr )
+	, BatchDataHDR(nullptr)
+	, CurrentCachedElementList(nullptr)
 	, PrecachedClippingStates(nullptr)
 	, RenderingPolicy( &InRenderingPolicy.Get() )
 	, NumPostProcessPasses(0)
@@ -287,7 +288,9 @@ void FSlateElementBatcher::AddElements(FSlateWindowElementList& WindowElementLis
 #endif
 
 	BatchData = &WindowElementList.GetBatchData();
+	BatchDataHDR = &WindowElementList.GetBatchDataHDR();
 	check(BatchData->GetRenderBatches().Num() == 0);
+	check(BatchDataHDR->GetRenderBatches().Num() == 0);
 
 
 	FVector2D ViewportSize = WindowElementList.GetPaintWindow()->GetViewportSize();
@@ -314,6 +317,7 @@ void FSlateElementBatcher::AddElements(FSlateWindowElementList& WindowElementLis
 
 	// Done with the element list
 	BatchData = nullptr;
+	BatchDataHDR = nullptr;
 	PrecachedClippingStates = nullptr;
 
 #if STATS
@@ -2174,8 +2178,13 @@ void FSlateElementBatcher::AddViewportElement( const FSlateDrawElement& DrawElem
 		// This is a slight hack, but the grayscale font shader is the same as the general shader except it reads alpha only textures and doesn't support tiling
 		ShaderType = ESlateShader::GrayscaleFont;
 	}
-
-	FSlateRenderBatch& RenderBatch = CreateRenderBatch( Layer, FShaderParams(), ViewportResource, ESlateDrawPrimitive::TriangleList, ShaderType, InDrawEffects, DrawFlags, DrawElement);
+	static const auto CVarHDROutputEnabled = IConsoleManager::Get().FindTConsoleVariableDataInt(TEXT("r.HDR.EnableHDROutput"));
+	static const auto CVarCompositeMode = IConsoleManager::Get().FindTConsoleVariableDataInt(TEXT("r.HDR.UI.CompositeMode"));
+	const bool bHDREnabled = CVarHDROutputEnabled && (CVarHDROutputEnabled->GetValueOnAnyThread() != 0);
+	const bool bCompositeModeEnabled = CVarCompositeMode && (CVarCompositeMode->GetValueOnAnyThread() != 0);
+	const bool bHDRBatch = bHDREnabled && bCompositeModeEnabled;
+	FSlateBatchData* UsedBatchData = bHDRBatch ? BatchDataHDR : BatchData;
+	FSlateRenderBatch& RenderBatch = CreateRenderBatch(UsedBatchData, Layer, FShaderParams(), ViewportResource, ESlateDrawPrimitive::TriangleList, ShaderType, InDrawEffects, DrawFlags, DrawElement);
 
 	// Tag this batch as requiring vsync if the viewport requires it.
 	if( ViewportResource != nullptr && !DrawElementPayload.bAllowViewportScaling )
@@ -2214,6 +2223,31 @@ void FSlateElementBatcher::AddViewportElement( const FSlateDrawElement& DrawElem
 	RenderBatch.AddIndex( IndexStart + 2 );
 	RenderBatch.AddIndex( IndexStart + 1 );
 	RenderBatch.AddIndex( IndexStart + 3 );
+
+	if (bHDRBatch)
+	{
+		// used to poke a hole in the slate tree: in case HDR is enabled, we need to compose the hdr scene with the SDR ui based on UI alpha
+		// The problem is that in editor mode, there's already a few quads already drawn below the viewport that we actually don't want. If we had an easy way to split
+		// in-game UI from tool UI, we could have prevented this with colorwritemask, but it doesn't seem to be the case. Instead, we force to draw
+		// an alpha=0 quad to allow the scene RT to be composited properly. AFAICT, game UI is rendered on top of it
+		DrawFlags = ESlateBatchDrawFlag::NoBlending;
+		const FColor TransparentBlackColor(0, 0, 0, 0);
+		FSlateRenderBatch& RenderBatch2 = CreateRenderBatch(Layer, FShaderParams(), nullptr, ESlateDrawPrimitive::TriangleList, ShaderType, InDrawEffects, DrawFlags, DrawElement);
+		RenderBatch2.AddVertex(FSlateVertex::Make<Rounding>(RenderTransform, FVector2f(TopLeft), FVector2f(0.0f, 0.0f), TransparentBlackColor));
+		RenderBatch2.AddVertex(FSlateVertex::Make<Rounding>(RenderTransform, FVector2f(TopRight), FVector2f(1.0f, 0.0f), TransparentBlackColor));
+		RenderBatch2.AddVertex(FSlateVertex::Make<Rounding>(RenderTransform, FVector2f(BotLeft), FVector2f(0.0f, 1.0f), TransparentBlackColor));
+		RenderBatch2.AddVertex(FSlateVertex::Make<Rounding>(RenderTransform, FVector2f(BotRight), FVector2f(1.0f, 1.0f), TransparentBlackColor));
+
+		// Add 6 indices to the vertex buffer.  (2 tri's per quad, 3 indices per tri)
+		RenderBatch2.AddIndex(IndexStart + 0);
+		RenderBatch2.AddIndex(IndexStart + 1);
+		RenderBatch2.AddIndex(IndexStart + 2);
+
+		RenderBatch2.AddIndex(IndexStart + 2);
+		RenderBatch2.AddIndex(IndexStart + 1);
+		RenderBatch2.AddIndex(IndexStart + 3);
+	}
+
 }
 
 template<ESlateVertexRounding Rounding>
@@ -2559,6 +2593,7 @@ void FSlateElementBatcher::AddPostProcessPass(const FSlateDrawElement& DrawEleme
 }
 
 FSlateRenderBatch& FSlateElementBatcher::CreateRenderBatch(
+	FSlateBatchData* SlateBatchData,
 	int32 Layer, 
 	const FShaderParams& ShaderParams,
 	const FSlateShaderResource* InResource,
@@ -2570,7 +2605,7 @@ FSlateRenderBatch& FSlateElementBatcher::CreateRenderBatch(
 {
 	FSlateRenderBatch& NewBatch = CurrentCachedElementList
 		? CurrentCachedElementList->AddRenderBatch(Layer, ShaderParams, InResource, PrimitiveType, ShaderType, DrawEffects, DrawFlags, DrawElement.GetSceneIndex())
-		: BatchData->AddRenderBatch(Layer, ShaderParams, InResource, PrimitiveType, ShaderType, DrawEffects, DrawFlags, DrawElement.GetSceneIndex());
+		: SlateBatchData->AddRenderBatch(Layer, ShaderParams, InResource, PrimitiveType, ShaderType, DrawEffects, DrawFlags, DrawElement.GetSceneIndex());
 
 	NewBatch.ClippingState = ResolveClippingState(DrawElement);
 
diff --git a/Engine/Source/Runtime/SlateCore/Public/Rendering/DrawElements.h b/Engine/Source/Runtime/SlateCore/Public/Rendering/DrawElements.h
--- a/Engine/Source/Runtime/SlateCore/Public/Rendering/DrawElements.h
+++ b/Engine/Source/Runtime/SlateCore/Public/Rendering/DrawElements.h
@@ -626,6 +626,7 @@ public:
 
 
 	FSlateBatchData& GetBatchData() { return BatchData; }
+	FSlateBatchData& GetBatchDataHDR() { return BatchDataHDR; }
 
 	SLATECORE_API void SetRenderTargetWindow(SWindow* InRenderTargetWindow);
 
@@ -657,6 +658,9 @@ private:
 	/** Batched data used for rendering */
 	FSlateBatchData BatchData;
 
+	/** Batched data used for rendering */
+	FSlateBatchData BatchDataHDR;
+
 	/**  */
 	FSlateClippingManager ClippingManager;
 
diff --git a/Engine/Source/Runtime/SlateCore/Public/Rendering/ElementBatcher.h b/Engine/Source/Runtime/SlateCore/Public/Rendering/ElementBatcher.h
--- a/Engine/Source/Runtime/SlateCore/Public/Rendering/ElementBatcher.h
+++ b/Engine/Source/Runtime/SlateCore/Public/Rendering/ElementBatcher.h
@@ -338,6 +338,20 @@ private:
 	void AddPostProcessPass(const FSlateDrawElement& DrawElement, const FVector2D& WindowSize);
 
 	FSlateRenderBatch& CreateRenderBatch(
+		int32 Layer,
+		const FShaderParams& ShaderParams,
+		const FSlateShaderResource* InResource,
+		ESlateDrawPrimitive PrimitiveType,
+		ESlateShader ShaderType,
+		ESlateDrawEffect DrawEffects,
+		ESlateBatchDrawFlag DrawFlags,
+		const FSlateDrawElement& DrawElement)
+	{
+		return CreateRenderBatch(BatchData, Layer, ShaderParams, InResource, PrimitiveType, ShaderType, DrawEffects, DrawFlags, DrawElement);
+	}
+
+	FSlateRenderBatch& CreateRenderBatch(
+		FSlateBatchData* SlateBatchData,
 		int32 Layer,
 		const FShaderParams& ShaderParams,
 		const FSlateShaderResource* InResource,
@@ -379,6 +393,7 @@ private:
 private:
 	/** Uncached Batch data currently being filled in */
 	FSlateBatchData* BatchData;
+	FSlateBatchData* BatchDataHDR;
 
 	/** Cached batches currently being filled in */
 	FSlateCachedElementList* CurrentCachedElementList;
diff --git a/Engine/Source/Runtime/SlateRHIRenderer/Private/SlateRHIRenderer.cpp b/Engine/Source/Runtime/SlateRHIRenderer/Private/SlateRHIRenderer.cpp
--- a/Engine/Source/Runtime/SlateRHIRenderer/Private/SlateRHIRenderer.cpp
+++ b/Engine/Source/Runtime/SlateRHIRenderer/Private/SlateRHIRenderer.cpp
@@ -711,6 +711,85 @@ SHADER_VARIATION(0)  SHADER_VARIATION(1)
 int32 SlateWireFrame = 0;
 static FAutoConsoleVariableRef CVarSlateWireframe(TEXT("Slate.ShowWireFrame"), SlateWireFrame, TEXT(""), ECVF_Default);
 
+void RenderSlateBatch(FTexture2DRHIRef SlateRenderTarget, bool bClear, bool bIsHDR, FViewportInfo& ViewportInfo, const FMatrix& ViewMatrix, FSlateBatchData& BatchData, FRHICommandListImmediate& RHICmdList,
+				      const uint32 ViewportWidth, const uint32 ViewportHeight, const struct FSlateDrawWindowCommandParams& DrawCommandParams, TSharedPtr<FSlateRHIRenderingPolicy> RenderingPolicy, 
+					  FTexture2DRHIRef PostProcessBuffer)
+{
+	FRHIRenderPassInfo RPInfo(SlateRenderTarget, ERenderTargetActions::Load_Store);
+
+	if (bClear)
+	{
+		RPInfo.ColorRenderTargets[0].Action = ERenderTargetActions::Clear_Store;
+	}
+
+	if (ViewportInfo.bRequiresStencilTest)
+	{
+		check(IsValidRef(ViewportInfo.DepthStencil));
+
+		ERenderTargetActions StencilAction = IsMemorylessTexture(ViewportInfo.DepthStencil) ? ERenderTargetActions::DontLoad_DontStore : ERenderTargetActions::DontLoad_Store;
+		RPInfo.DepthStencilRenderTarget.Action = MakeDepthStencilTargetActions(ERenderTargetActions::DontLoad_DontStore, StencilAction);
+		RPInfo.DepthStencilRenderTarget.DepthStencilTarget = ViewportInfo.DepthStencil;
+		RPInfo.DepthStencilRenderTarget.ExclusiveDepthStencil = FExclusiveDepthStencil::DepthNop_StencilWrite;
+	}
+
+#if WITH_SLATE_VISUALIZERS
+	if (CVarShowSlateBatching.GetValueOnRenderThread() != 0 || CVarShowSlateOverdraw.GetValueOnRenderThread() != 0)
+	{
+		RPInfo.ColorRenderTargets[0].Action = ERenderTargetActions::Clear_Store;
+		if (ViewportInfo.bRequiresStencilTest)
+		{
+			// Reset the backbuffer as our color render target and also set a depth stencil buffer
+			ERenderTargetActions StencilAction = IsMemorylessTexture(ViewportInfo.DepthStencil) ? ERenderTargetActions::Clear_DontStore : ERenderTargetActions::Clear_Store;
+			RPInfo.DepthStencilRenderTarget.Action = MakeDepthStencilTargetActions(ERenderTargetActions::Load_Store, StencilAction);
+			RPInfo.DepthStencilRenderTarget.DepthStencilTarget = ViewportInfo.DepthStencil;
+			RPInfo.DepthStencilRenderTarget.ExclusiveDepthStencil = FExclusiveDepthStencil::DepthWrite_StencilWrite;
+		}
+	}
+#endif
+	{
+		bool bHasBatches = BatchData.GetRenderBatches().Num() > 0;
+		if (bHasBatches || bClear)
+		{
+			TransitionRenderPassTargets(RHICmdList, RPInfo);
+			RHICmdList.BeginRenderPass(RPInfo, TEXT("SlateBatches"));
+			SCOPE_CYCLE_COUNTER(STAT_SlateRTDrawBatches);
+
+			if (bHasBatches)
+			{
+				FSlateBackBuffer SlateBackBuffer(SlateRenderTarget, FIntPoint(ViewportWidth, ViewportHeight));
+
+				FSlateRenderingParams RenderParams(ViewMatrix * ViewportInfo.ProjectionMatrix, DrawCommandParams.Time);
+				RenderParams.bWireFrame = !!SlateWireFrame;
+				RenderParams.bIsHDR = bIsHDR;
+				RenderingPolicy->SetUseGammaCorrection(!bIsHDR);
+
+				FTexture2DRHIRef EmptyTarget;
+
+				RenderingPolicy->DrawElements
+				(
+					RHICmdList,
+					SlateBackBuffer,
+					SlateRenderTarget,
+					PostProcessBuffer,
+					ViewportInfo.bRequiresStencilTest ? ViewportInfo.DepthStencil : EmptyTarget,
+					BatchData.GetFirstRenderBatchIndex(),
+					BatchData.GetRenderBatches(),
+					RenderParams
+				);
+			}
+		}
+	}
+
+	// @todo Could really use a refactor.
+	// Kind of gross but we don't want to restart renderpasses for no reason.
+	// If the color deficiency shaders are active within DrawElements there will not be a renderpass here.
+	// In the general case there will be a RenderPass active at this point.
+	if (RHICmdList.IsInsideRenderPass())
+	{
+		RHICmdList.EndRenderPass();
+	}
+}
+
 /** Draws windows from a FSlateDrawBuffer on the render thread */
 void FSlateRHIRenderer::DrawWindow_RenderThread(FRHICommandListImmediate& RHICmdList, FViewportInfo& ViewportInfo, FSlateWindowElementList& WindowElementList, const struct FSlateDrawWindowCommandParams& DrawCommandParams)
 {
@@ -887,83 +966,20 @@ void FSlateRHIRenderer::DrawWindow_RenderThread(FRHICommandListImmediate& RHICmd
 			RHICmdList.BeginDrawingViewport(ViewportInfo.ViewportRHI, FTextureRHIRef());
 			RHICmdList.SetViewport(0, 0, 0, ViewportWidth, ViewportHeight, 0.0f);
 
-			{
-				FRHIRenderPassInfo RPInfo(BackBuffer, ERenderTargetActions::Load_Store);
-
-				if (bClear)
-				{
-					RPInfo.ColorRenderTargets[0].Action = ERenderTargetActions::Clear_Store;
-				}
-
-				if (ViewportInfo.bRequiresStencilTest)
-				{
-					check(IsValidRef(ViewportInfo.DepthStencil));
-
-					ERenderTargetActions StencilAction = IsMemorylessTexture(ViewportInfo.DepthStencil) ? ERenderTargetActions::DontLoad_DontStore : ERenderTargetActions::DontLoad_Store;
-					RPInfo.DepthStencilRenderTarget.Action = MakeDepthStencilTargetActions(ERenderTargetActions::DontLoad_DontStore, StencilAction);
-					RPInfo.DepthStencilRenderTarget.DepthStencilTarget = ViewportInfo.DepthStencil;
-					RPInfo.DepthStencilRenderTarget.ExclusiveDepthStencil = FExclusiveDepthStencil::DepthNop_StencilWrite;
-				}
-
-	#if WITH_SLATE_VISUALIZERS
-				if (CVarShowSlateBatching.GetValueOnRenderThread() != 0 || CVarShowSlateOverdraw.GetValueOnRenderThread() != 0)
-				{
-					RPInfo.ColorRenderTargets[0].Action = ERenderTargetActions::Clear_Store;
-					if (ViewportInfo.bRequiresStencilTest)
-					{
-						// Reset the backbuffer as our color render target and also set a depth stencil buffer
-						ERenderTargetActions StencilAction = IsMemorylessTexture(ViewportInfo.DepthStencil) ? ERenderTargetActions::Clear_DontStore : ERenderTargetActions::Clear_Store;
-						RPInfo.DepthStencilRenderTarget.Action = MakeDepthStencilTargetActions(ERenderTargetActions::Load_Store, StencilAction);
-						RPInfo.DepthStencilRenderTarget.DepthStencilTarget = ViewportInfo.DepthStencil;
-						RPInfo.DepthStencilRenderTarget.ExclusiveDepthStencil = FExclusiveDepthStencil::DepthWrite_StencilWrite;
-					}
-				}
-	#endif
-				{
-					bool bHasBatches = BatchData.GetRenderBatches().Num() > 0;
-					if (bHasBatches || bClear)
-					{
-						TransitionRenderPassTargets(RHICmdList, RPInfo);
-						RHICmdList.BeginRenderPass(RPInfo, TEXT("SlateBatches"));
-						SCOPE_CYCLE_COUNTER(STAT_SlateRTDrawBatches);
-
-						if (bHasBatches)
-						{
-							FSlateBackBuffer BackBufferTarget(BackBuffer, FIntPoint(ViewportWidth, ViewportHeight));
-
-							FSlateRenderingParams RenderParams(ViewMatrix * ViewportInfo.ProjectionMatrix, DrawCommandParams.Time);
-							RenderParams.bWireFrame = !!SlateWireFrame;
-							RenderParams.bIsHDR = ViewportInfo.bHDREnabled;
-
-							FTexture2DRHIRef EmptyTarget;
-
-							RenderingPolicy->DrawElements
-							(
-								RHICmdList,
-								BackBufferTarget,
-								BackBuffer,
-								PostProcessBuffer,
-								ViewportInfo.bRequiresStencilTest ? ViewportInfo.DepthStencil : EmptyTarget,
-								BatchData.GetFirstRenderBatchIndex(),
-								BatchData.GetRenderBatches(),
-								RenderParams
-							);
-						}
-					}
-				}
+			bool bHdrTarget = IsHDREnabled() && !bCompositeUI;
+			RenderSlateBatch(BackBuffer, bClear, bHdrTarget, ViewportInfo, ViewMatrix, BatchData, RHICmdList, ViewportWidth, ViewportHeight, DrawCommandParams, RenderingPolicy, PostProcessBuffer);
 
-				// @todo Could really use a refactor.
-				// Kind of gross but we don't want to restart renderpasses for no reason.
-				// If the color deficiency shaders are active within DrawElements there will not be a renderpass here.
-				// In the general case there will be a RenderPass active at this point.
-				if (RHICmdList.IsInsideRenderPass())
+			if (bCompositeUI)
+			{
+				FSlateBatchData& BatchDataHDR = WindowElementList.GetBatchDataHDR();
+				bool bHasBatches = BatchDataHDR.GetRenderBatches().Num() > 0;
+				if (bHasBatches)
 				{
-					RHICmdList.EndRenderPass();
+					RenderingPolicy->BuildRenderingBuffers(RHICmdList, BatchDataHDR);
+					FTexture2DRHIRef UITargetHDRRTRHI(ViewportInfo.HDRSourceRT->GetRHI());
+					RenderSlateBatch(UITargetHDRRTRHI, /*bClear*/ false, /*bIsHDR*/ true, ViewportInfo, ViewMatrix, BatchDataHDR, RHICmdList, ViewportWidth, ViewportHeight, DrawCommandParams, RenderingPolicy, PostProcessBuffer);
 				}
-			}
 
-			if (bCompositeUI)
-			{
 				SCOPED_DRAW_EVENT(RHICmdList, SlateUI_Composition);
 
 				static const FName RendererModuleName("Renderer");
