fix for editor when r.HDR.EnableHDROutput would be true, but GRHISupportsHDROutput is false (typically happens if no HDR monitor connected).

diff --git a/Engine/Source/Runtime/SlateCore/Private/Rendering/ElementBatcher.cpp b/Engine/Source/Runtime/SlateCore/Private/Rendering/ElementBatcher.cpp
--- a/Engine/Source/Runtime/SlateCore/Private/Rendering/ElementBatcher.cpp
+++ b/Engine/Source/Runtime/SlateCore/Private/Rendering/ElementBatcher.cpp
@@ -264,6 +264,7 @@ FSlateElementBatcher::FSlateElementBatcher( TSharedRef<FSlateRenderingPolicy> In
 	, PixelCenterOffset( InRenderingPolicy->GetPixelCenterOffset() )
 	, bSRGBVertexColor( !InRenderingPolicy->IsVertexColorInLinearSpace() )
 	, bRequiresVsync(false)
+	, bCompositeHDRViewports(false)
 {
 }
 
@@ -2178,12 +2179,7 @@ void FSlateElementBatcher::AddViewportElement( const FSlateDrawElement& DrawElem
 		// This is a slight hack, but the grayscale font shader is the same as the general shader except it reads alpha only textures and doesn't support tiling
 		ShaderType = ESlateShader::GrayscaleFont;
 	}
-	static const auto CVarHDROutputEnabled = IConsoleManager::Get().FindTConsoleVariableDataInt(TEXT("r.HDR.EnableHDROutput"));
-	static const auto CVarCompositeMode = IConsoleManager::Get().FindTConsoleVariableDataInt(TEXT("r.HDR.UI.CompositeMode"));
-	const bool bHDREnabled = CVarHDROutputEnabled && (CVarHDROutputEnabled->GetValueOnAnyThread() != 0);
-	const bool bCompositeModeEnabled = CVarCompositeMode && (CVarCompositeMode->GetValueOnAnyThread() != 0);
-	const bool bHDRBatch = bHDREnabled && bCompositeModeEnabled;
-	FSlateBatchData* UsedBatchData = bHDRBatch ? BatchDataHDR : BatchData;
+	FSlateBatchData* UsedBatchData = bCompositeHDRViewports ? BatchDataHDR : BatchData;
 	FSlateRenderBatch& RenderBatch = CreateRenderBatch(UsedBatchData, Layer, FShaderParams(), ViewportResource, ESlateDrawPrimitive::TriangleList, ShaderType, InDrawEffects, DrawFlags, DrawElement);
 
 	// Tag this batch as requiring vsync if the viewport requires it.
@@ -2224,7 +2220,7 @@ void FSlateElementBatcher::AddViewportElement( const FSlateDrawElement& DrawElem
 	RenderBatch.AddIndex( IndexStart + 1 );
 	RenderBatch.AddIndex( IndexStart + 3 );
 
-	if (bHDRBatch)
+	if (bCompositeHDRViewports)
 	{
 		// used to poke a hole in the slate tree: in case HDR is enabled, we need to compose the hdr scene with the SDR ui based on UI alpha
 		// The problem is that in editor mode, there's already a few quads already drawn below the viewport that we actually don't want. If we had an easy way to split
@@ -2908,6 +2904,7 @@ void FSlateElementBatcher::BuildShapedTextSequence(const FShapedTextBuildContext
 void FSlateElementBatcher::ResetBatches()
 {
 	bRequiresVsync = false;
+	bCompositeHDRViewports = false;
 	NumPostProcessPasses = 0;
 }
 
diff --git a/Engine/Source/Runtime/SlateCore/Public/Rendering/ElementBatcher.h b/Engine/Source/Runtime/SlateCore/Public/Rendering/ElementBatcher.h
--- a/Engine/Source/Runtime/SlateCore/Public/Rendering/ElementBatcher.h
+++ b/Engine/Source/Runtime/SlateCore/Public/Rendering/ElementBatcher.h
@@ -262,6 +262,10 @@ public:
 	/** Whether or not any post process passes were batched */
 	bool HasFXPassses() const { return NumPostProcessPasses > 0;}
 
+	bool CompositeHDRViewports() const { return bCompositeHDRViewports; }
+
+	void SetCompositeHDRViewports(bool bInCompositeHDRViewports) { bCompositeHDRViewports = bInCompositeHDRViewports; }
+
 	/** 
 	 * Resets all stored data accumulated during the batching process
 	 */
@@ -437,4 +441,7 @@ private:
 
 	// true if any element in the batch requires vsync.
 	bool bRequiresVsync;
+
+	// true if viewports get composited as a separate pass, instead of being rendered directly to the render target. Useful for HDR displays
+	bool bCompositeHDRViewports;
 };
diff --git a/Engine/Source/Runtime/SlateRHIRenderer/Private/SlateRHIRenderer.cpp b/Engine/Source/Runtime/SlateRHIRenderer/Private/SlateRHIRenderer.cpp
--- a/Engine/Source/Runtime/SlateRHIRenderer/Private/SlateRHIRenderer.cpp
+++ b/Engine/Source/Runtime/SlateRHIRenderer/Private/SlateRHIRenderer.cpp
@@ -791,6 +791,19 @@ void RenderSlateBatch(FTexture2DRHIRef SlateRenderTarget, bool bClear, bool bIsH
 	}
 }
 
+inline bool CompositeUIWithHdrRenderTarget()
+{
+	// Optional off-screen UI composition during HDR rendering
+	static const auto CVarCompositeMode = IConsoleManager::Get().FindTConsoleVariableDataInt(TEXT("r.HDR.UI.CompositeMode"));
+
+	const bool bSupportsUIComposition = GRHISupportsHDROutput && GSupportsVolumeTextureRendering && SupportsUICompositionRendering(GetFeatureLevelShaderPlatform(GMaxRHIFeatureLevel));
+	const bool bCompositeUI = bSupportsUIComposition
+		&& CVarCompositeMode && CVarCompositeMode->GetValueOnAnyThread() != 0
+		&& IsHDREnabled();
+
+	return bCompositeUI;
+}
+
 /** Draws windows from a FSlateDrawBuffer on the render thread */
 void FSlateRHIRenderer::DrawWindow_RenderThread(FRHICommandListImmediate& RHICmdList, FViewportInfo& ViewportInfo, FSlateWindowElementList& WindowElementList, const struct FSlateDrawWindowCommandParams& DrawCommandParams)
 {
@@ -821,13 +834,7 @@ void FSlateRHIRenderer::DrawWindow_RenderThread(FRHICommandListImmediate& RHICmd
 		FMaterialRenderProxy::UpdateDeferredCachedUniformExpressions();
 		GetRendererModule().InitializeSystemTextures(RHICmdList);
 
-		// Optional off-screen UI composition during HDR rendering
-		static const auto CVarCompositeMode = IConsoleManager::Get().FindTConsoleVariableDataInt(TEXT("r.HDR.UI.CompositeMode"));
-
-		const bool bSupportsUIComposition = GRHISupportsHDROutput && GSupportsVolumeTextureRendering && SupportsUICompositionRendering(GetFeatureLevelShaderPlatform(GMaxRHIFeatureLevel));
-		const bool bCompositeUI = bSupportsUIComposition
-			&& CVarCompositeMode && CVarCompositeMode->GetValueOnRenderThread() != 0
-			&& IsHDREnabled();
+		const bool bCompositeUI = CompositeUIWithHdrRenderTarget();
 
 		const int32 CompositionLUTSize = 32;
 
@@ -1321,9 +1328,11 @@ void FSlateRHIRenderer::DrawWindows_Private(FSlateDrawBuffer& WindowDrawBuffer)
 					}
 				}
 
-
+				bool bCompositeHDRViewports = ElementBatcher->CompositeHDRViewports();
+				ElementBatcher->SetCompositeHDRViewports(CompositeUIWithHdrRenderTarget());
 				// Add all elements for this window to the element batcher
 				ElementBatcher->AddElements(ElementList);
+				ElementBatcher->SetCompositeHDRViewports(bCompositeHDRViewports);
 
 				// Update the font cache with new text after elements are batched
 				FontCache->UpdateCache();
