FPostProcessRectParams SourceTexture support for 2DArray.
-FPostProcessRectParams was still using the old FTexture2DRHIRef/GetTexture2D pattern rather than using FTextureRHIRef directly.  This meant that it did not support 2DArray textures, like the VR swapchain, and caused an assert during a post-process blur on Lyra.

diff --git a/Engine/Source/Runtime/SlateRHIRenderer/Private/SlatePostProcessor.h b/Engine/Source/Runtime/SlateRHIRenderer/Private/SlatePostProcessor.h
--- a/Engine/Source/Runtime/SlateRHIRenderer/Private/SlatePostProcessor.h
+++ b/Engine/Source/Runtime/SlateRHIRenderer/Private/SlatePostProcessor.h
@@ -11,7 +11,7 @@ class IRendererModule;
 
 struct FPostProcessRectParams
 {
-	FTexture2DRHIRef SourceTexture;
+	FTextureRHIRef SourceTexture;
 	FSlateRect SourceRect;
 	FSlateRect DestRect;
 	FVector4 CornerRadius;
diff --git a/Engine/Source/Runtime/SlateRHIRenderer/Private/SlateRHIRenderingPolicy.cpp b/Engine/Source/Runtime/SlateRHIRenderer/Private/SlateRHIRenderingPolicy.cpp
--- a/Engine/Source/Runtime/SlateRHIRenderer/Private/SlateRHIRenderingPolicy.cpp
+++ b/Engine/Source/Runtime/SlateRHIRenderer/Private/SlateRHIRenderingPolicy.cpp
@@ -1184,7 +1184,7 @@ void FSlateRHIRenderingPolicy::DrawElements(
 				const FVector4f QuadPositionData = ShaderParams.PixelParams;
 
 				FPostProcessRectParams RectParams;
-				RectParams.SourceTexture = PostProcessTexture->GetTexture2D();
+				RectParams.SourceTexture = PostProcessTexture;
 				RectParams.SourceRect = FSlateRect(0, 0, PostProcessTexture->GetSizeX(), PostProcessTexture->GetSizeY());
 				RectParams.DestRect = FSlateRect(QuadPositionData.X, QuadPositionData.Y, QuadPositionData.Z, QuadPositionData.W);
 				RectParams.SourceTextureSize = PostProcessTexture->GetSizeXY();
