hdr multi window/display support
This improves how the editor deals with displays that have different capabilities (some with HDR, others without).
Centralize where the CVars are accessed
Have the swapchain created in SDR when shown on a SDR display
Make sure color gamut matches the one supported by the platform/swapchain.
Allow Movie scene capture to override the gamut/output format to keep existing behaviour (MovieSceneCapture relies on swapchain data)

diff --git a/Engine/Source/Runtime/D3D12RHI/Private/D3D12RHIPrivate.h b/Engine/Source/Runtime/D3D12RHI/Private/D3D12RHIPrivate.h
--- a/Engine/Source/Runtime/D3D12RHI/Private/D3D12RHIPrivate.h
+++ b/Engine/Source/Runtime/D3D12RHI/Private/D3D12RHIPrivate.h
@@ -32,6 +32,7 @@
 #include "ID3D12DynamicRHI.h"
 #include "GPUProfiler.h"
 #include "ShaderCore.h"
+#include "HDRHelper.h"
 
 DECLARE_LOG_CATEGORY_EXTERN(LogD3D12RHI, Log, All);
 DECLARE_LOG_CATEGORY_EXTERN(LogD3D12GapRecorder, Log, All);
@@ -393,6 +394,7 @@ public:
 	virtual FTexture2DRHIRef RHIGetViewportBackBuffer(FRHIViewport* Viewport) final override;
 	virtual void RHIAliasTextureResources(FTextureRHIRef& DestTexture, FTextureRHIRef& SrcTexture) final override;
 	virtual FTextureRHIRef RHICreateAliasedTexture(FTextureRHIRef& SourceTexture) final override;
+	virtual void RHIGetDisplaysInformation(FDisplayInformationArray& OutDisplayInformation) final override;
 	virtual void RHIAdvanceFrameFence() final override;
 	virtual void RHIAdvanceFrameForGetViewportBackBuffer(FRHIViewport* Viewport) final override;
 	virtual void RHIAcquireThreadOwnership() final override;
@@ -1200,6 +1202,8 @@ public:
 
 	FRHIBuffer* CreateBuffer(const FRHIBufferCreateInfo& CreateInfo, const TCHAR* DebugName, ERHIAccess InitialState, ID3D12ResourceAllocator* ResourceAllocator);
 
+	bool SetupDisplayHDRMetaData();
+
 protected:
 
 	FD3D12Texture* CreateTextureFromResource(bool bTextureArray, bool bCubeTexture, EPixelFormat Format, ETextureCreateFlags TexCreateFlags, const FClearValueBinding& ClearValueBinding, ID3D12Resource* Resource);
@@ -1240,6 +1244,8 @@ protected:
 	HANDLE FlipEvent;
 
 	const bool bAllowVendorDevice;
+
+	FDisplayInformationArray DisplayList;
 };
 
 PRAGMA_ENABLE_DEPRECATION_WARNINGS
diff --git a/Engine/Source/Runtime/D3D12RHI/Private/D3D12Viewport.cpp b/Engine/Source/Runtime/D3D12RHI/Private/D3D12Viewport.cpp
--- a/Engine/Source/Runtime/D3D12RHI/Private/D3D12Viewport.cpp
+++ b/Engine/Source/Runtime/D3D12RHI/Private/D3D12Viewport.cpp
@@ -7,6 +7,7 @@
 #include "D3D12RHIPrivate.h"
 #include "RenderCore.h"
 #include "Engine/RendererSettings.h"
+#include "HDRHelper.h"
 
 namespace D3D12RHI
 {
@@ -568,10 +569,19 @@ void FD3D12Viewport::Resize(uint32 InSizeX, uint32 InSizeY, bool bInIsFullscreen
 		}
 	}
 
+	RECT WindowRect = {};
+#if PLATFORM_WINDOWS
+	GetWindowRect(WindowHandle, &WindowRect);
+#endif
+	FVector2D WindowTopLeft((float)WindowRect.left, (float)WindowRect.top);
+	FVector2D WindowBottomRight((float)WindowRect.right, (float)WindowRect.bottom);
+	bool bHDREnabled;
+	HDRGetMetaData(DisplayOutputFormat, DisplayColorGamut, bHDREnabled, WindowTopLeft, WindowBottomRight, (void*)WindowHandle);
+
 	ResizeInternal();
 
 	// Enable HDR if desired.
-	if (CheckHDRSupport())
+	if (bHDREnabled)
 	{
 		EnableHDR();
 	}
@@ -910,7 +920,7 @@ void FD3D12Viewport::IssueFrameEvent()
 
 bool FD3D12Viewport::CheckHDRSupport()
 {
-	return GRHISupportsHDROutput && IsHDREnabled();
+	return IsHDREnabled();
 }
 
 EPixelFormat GetDefaultBackBufferPixelFormat()
diff --git a/Engine/Source/Runtime/D3D12RHI/Private/Windows/WindowsD3D12Device.cpp b/Engine/Source/Runtime/D3D12RHI/Private/Windows/WindowsD3D12Device.cpp
--- a/Engine/Source/Runtime/D3D12RHI/Private/Windows/WindowsD3D12Device.cpp
+++ b/Engine/Source/Runtime/D3D12RHI/Private/Windows/WindowsD3D12Device.cpp
@@ -365,16 +365,18 @@ static bool SupportsDepthBoundsTest(FD3D12DynamicRHI* D3DRHI)
 	return D3DRHI->GetAdapter().IsDepthBoundsTestSupported();
 }
 
-static bool SupportsHDROutput(FD3D12DynamicRHI* D3DRHI)
+bool FD3D12DynamicRHI::SetupDisplayHDRMetaData()
 {
 	// Determines if any displays support HDR
-	check(D3DRHI && D3DRHI->GetNumAdapters() >= 1);
+	check(GetNumAdapters() >= 1);
+
+	DisplayList.Empty();
 
 	bool bSupportsHDROutput = false;
-	const int32 NumAdapters = D3DRHI->GetNumAdapters();
+	const int32 NumAdapters = GetNumAdapters();
 	for (int32 AdapterIndex = 0; AdapterIndex < NumAdapters; ++AdapterIndex)
 	{
-		FD3D12Adapter& Adapter = D3DRHI->GetAdapter(AdapterIndex);
+		FD3D12Adapter& Adapter = GetAdapter(AdapterIndex);
 		IDXGIAdapter* DXGIAdapter = Adapter.GetAdapter();
 
 		for (uint32 DisplayIndex = 0; true; ++DisplayIndex)
@@ -402,6 +404,11 @@ static bool SupportsHDROutput(FD3D12DynamicRHI* D3DRHI)
 
 					bSupportsHDROutput = true;
 				}
+				FDisplayInformation DisplayInformation{};
+				DisplayInformation.bHDRSupported = bDisplaySupportsHDROutput;
+				const RECT& DisplayCoords = OutputDesc.DesktopCoordinates;
+				DisplayInformation.DesktopCoordinates = FIntRect(DisplayCoords.left, DisplayCoords.top, DisplayCoords.right, DisplayCoords.bottom);
+				DisplayList.Add(DisplayInformation);
 			}
 		}
 	}
@@ -1401,12 +1408,12 @@ void FD3D12DynamicRHI::Init()
 	GSupportsDepthBoundsTest = SupportsDepthBoundsTest(this);
 
 	{
-		GRHISupportsHDROutput = SupportsHDROutput(this);
+		GRHISupportsHDROutput = SetupDisplayHDRMetaData();
 
 		// Specify the desired HDR pixel format.
 		// Possible values are:
 		//	1) PF_FloatRGBA - FP16 format that allows for linear gamma. This is the current engine default.
-		//					r.HDR.Display.ColorGamut = 2 (Rec2020 / BT2020)
+		//					r.HDR.Display.ColorGamut = 0 (sRGB which is the same gamut as ScRGB)
 		//					r.HDR.Display.OutputDevice = 5 or 6 (ScRGB)
 		//	2) PF_A2B10G10R10 - Save memory vs FP16 as well as allow for possible performance improvements 
 		//						in fullscreen by avoiding format conversions.
diff --git a/Engine/Source/Runtime/D3D12RHI/Private/Windows/WindowsD3D12Viewport.cpp b/Engine/Source/Runtime/D3D12RHI/Private/Windows/WindowsD3D12Viewport.cpp
--- a/Engine/Source/Runtime/D3D12RHI/Private/Windows/WindowsD3D12Viewport.cpp
+++ b/Engine/Source/Runtime/D3D12RHI/Private/Windows/WindowsD3D12Viewport.cpp
@@ -52,6 +52,8 @@ FD3D12Viewport::FD3D12Viewport(class FD3D12Adapter* InParent, HWND InWindowHandl
 	SDRDummyBackBuffer_RenderThread(nullptr),
 	SDRBackBuffer_RHIThread(nullptr),
 	SDRPixelFormat(PF_B8G8R8A8),
+	DisplayColorGamut(EDisplayColorGamut::sRGB_D65),
+	DisplayOutputFormat(EDisplayOutputFormat::SDR_sRGB),
 	Fence(InParent, FRHIGPUMask::All(), L"Viewport Fence"),
 	LastSignaledValue(0)
 #if WITH_MGPU
@@ -397,27 +399,22 @@ void FD3D12Viewport::EnableHDR()
 {
 	if ( GRHISupportsHDROutput && IsHDREnabled() )
 	{
-		static const auto CVarHDROutputDevice = IConsoleManager::Get().FindTConsoleVariableDataInt(TEXT("r.HDR.Display.OutputDevice"));
-		const EDisplayOutputFormat OutputDevice = EDisplayOutputFormat(CVarHDROutputDevice->GetValueOnAnyThread());
-
-		const float DisplayMaxOutputNits = (OutputDevice == EDisplayOutputFormat::HDR_ACES_2000nit_ST2084 || OutputDevice == EDisplayOutputFormat::HDR_ACES_2000nit_ScRGB) ? 2000.f : 1000.f;
+		const float DisplayMaxOutputNits = (DisplayOutputFormat == EDisplayOutputFormat::HDR_ACES_2000nit_ST2084 || DisplayOutputFormat == EDisplayOutputFormat::HDR_ACES_2000nit_ScRGB) ? 2000.f : 1000.f;
 		const float DisplayMinOutputNits = 0.0f;	// Min output of the display
 		const float DisplayMaxCLL = 0.0f;			// Max content light level in lumens (0.0 == unknown)
 		const float DisplayFALL = 0.0f;				// Frame average light level (0.0 == unknown)
 
 		// Ideally we can avoid setting TV meta data and instead the engine can do tone mapping based on the
 		// actual current display properties (display mapping).
-		static const auto CVarHDRColorGamut = IConsoleManager::Get().FindTConsoleVariableDataInt(TEXT("r.HDR.Display.ColorGamut"));
-		const EDisplayColorGamut DisplayGamut = EDisplayColorGamut(CVarHDRColorGamut->GetValueOnAnyThread());
 		SetHDRTVMode(true,
-			DisplayGamut,
+			DisplayColorGamut,
 			DisplayMaxOutputNits,
 			DisplayMinOutputNits,
 			DisplayMaxCLL,
 			DisplayFALL);
 
 		// Ensure we have the correct color space set.
-		EnsureColorSpace(DisplayGamut, OutputDevice);
+		EnsureColorSpace(DisplayColorGamut, DisplayOutputFormat);
 	}
 }
 
@@ -542,13 +539,18 @@ void FD3D12Viewport::EnsureColorSpace(EDisplayColorGamut DisplayGamut, EDisplayO
 	if (ColorSpace != NewColorSpace)
 	{
 		uint32 ColorSpaceSupport = 0;
-		if (SUCCEEDED(SwapChain4->CheckColorSpaceSupport(NewColorSpace, &ColorSpaceSupport)) &&
-			((ColorSpaceSupport & DXGI_SWAP_CHAIN_COLOR_SPACE_SUPPORT_FLAG_PRESENT) != 0))
+		HRESULT hr = SwapChain4->CheckColorSpaceSupport(NewColorSpace, &ColorSpaceSupport);
+		FString NewColorSpaceName = GetDXGIColorSpaceString(NewColorSpace);
+		if (SUCCEEDED(hr) && ((ColorSpaceSupport & DXGI_SWAP_CHAIN_COLOR_SPACE_SUPPORT_FLAG_PRESENT) != 0))
 		{
 			VERIFYD3D12RESULT(SwapChain4->SetColorSpace1(NewColorSpace));
-			UE_LOG(LogD3D12RHI, Log, TEXT("Setting color space on swap chain (%#016llx): %s"), SwapChain4.GetReference(), *GetDXGIColorSpaceString(NewColorSpace));
+			UE_LOG(LogD3D12RHI, Log, TEXT("Setting color space on swap chain (%#016llx): %s"), SwapChain4.GetReference(), *NewColorSpaceName);
 			ColorSpace = NewColorSpace;
 		}
+		else
+		{
+			UE_LOG(LogD3D12RHI, Error, TEXT("Warning: unabled to set color space %s to the swapchain: verify EDisplayOutputFormat / swapchain format"), *NewColorSpaceName);
+		}
 	}
 }
 
@@ -612,4 +614,9 @@ void FD3D12Viewport::OnResumeRendering()
 void FD3D12Viewport::OnSuspendRendering()
 {}
 
+void FD3D12DynamicRHI::RHIGetDisplaysInformation(FDisplayInformationArray& OutDisplayInformation)
+{
+	OutDisplayInformation.Append(DisplayList);
+}
+
 #include "Windows/HideWindowsPlatformTypes.h"
diff --git a/Engine/Source/Runtime/D3D12RHI/Public/D3D12Viewport.h b/Engine/Source/Runtime/D3D12RHI/Public/D3D12Viewport.h
--- a/Engine/Source/Runtime/D3D12RHI/Public/D3D12Viewport.h
+++ b/Engine/Source/Runtime/D3D12RHI/Public/D3D12Viewport.h
@@ -243,6 +243,8 @@ private:
 	TRefCountPtr<FD3D12Texture> SDRDummyBackBuffer_RenderThread;
 	FD3D12Texture* SDRBackBuffer_RHIThread;
 	EPixelFormat SDRPixelFormat;
+	EDisplayColorGamut DisplayColorGamut;
+	EDisplayOutputFormat DisplayOutputFormat;
 
 	/** A fence value used to track the GPU's progress. */
 	FD3D12Fence Fence;
diff --git a/Engine/Source/Runtime/Engine/Private/Slate/SceneViewport.cpp b/Engine/Source/Runtime/Engine/Private/Slate/SceneViewport.cpp
--- a/Engine/Source/Runtime/Engine/Private/Slate/SceneViewport.cpp
+++ b/Engine/Source/Runtime/Engine/Private/Slate/SceneViewport.cpp
@@ -21,6 +21,7 @@
 #include "IHeadMountedDisplay.h"
 #include "IXRTrackingSystem.h"
 #include "StereoRenderTargetManager.h"
+#include "HDRHelper.h"
 
 DEFINE_LOG_CATEGORY(LogViewport);
 
@@ -63,6 +64,8 @@ FSceneViewport::FSceneViewport( FViewportClient* InViewportClient, TSharedPtr<SV
 	, CurrentBufferedTargetIndex(0)
 	, NextBufferedTargetIndex(0)
 	, NumTouches(0)
+	, DisplayColorGamut(EDisplayColorGamut::sRGB_D65)
+	, DisplayOutputFormat(EDisplayOutputFormat::SDR_sRGB)
 {
 	bIsSlateViewport = true;
 	ViewportType = NAME_SceneViewport;
@@ -369,14 +372,30 @@ void FSceneViewport::OnDrawViewport( const FGeometry& AllottedGeometry, const FS
 	{
 		FIntPoint DrawSize = FIntPoint( FMath::RoundToInt( AllottedGeometry.GetDrawSize().X ), FMath::RoundToInt( AllottedGeometry.GetDrawSize().Y ) );
 		bool bIsHDREnabled = IsHDREnabled();
-		if( GetSizeXY() != DrawSize || bIsHDREnabled != bHDRViewport)
+
+		SWindow* PaintWindow = OutDrawElements.GetPaintWindow();
+		bool bHDRStale = (bIsHDREnabled != bHDRViewport);
+		if (PaintWindow)
+		{
+			bool bNewHDREnabled;
+			EDisplayColorGamut NewDisplayColorGamut;
+			EDisplayOutputFormat NewDisplayOutputFormat;
+			HDRGetMetaData(NewDisplayOutputFormat, NewDisplayColorGamut, bNewHDREnabled, PaintWindow->GetPositionInScreen(), PaintWindow->GetPositionInScreen() + PaintWindow->GetSizeInScreen(), PaintWindow->GetNativeWindow()->GetOSWindowHandle());
+			// if we manage to get data for the window, we can ignore the global toggle IsHDREnabled since HDRGetMetaData will take both the global flag and the monitor properties
+			bHDRStale = DisplayOutputFormat != NewDisplayOutputFormat;
+			bHDRStale |= DisplayColorGamut != NewDisplayColorGamut;
+			bHDRStale |= bHDRViewport != bNewHDREnabled;
+		}
+
+	    if (GetSizeXY() != DrawSize || bHDRStale)
 		{
-			bHDRViewport = bIsHDREnabled;
-			TSharedPtr<SWindow> Window = FSlateApplication::Get().FindWidgetWindow( ViewportWidget.Pin().ToSharedRef() );
-			if ( Window.IsValid() )
+			TSharedPtr<SWindow> Window = FSlateApplication::Get().FindWidgetWindow(ViewportWidget.Pin().ToSharedRef());
+			if (Window.IsValid())
 			{
 				//@HACK VREDITOR
 				//check(Window.IsValid());
+                // This makes sure that by the time when we call HDRGetMetaData again in UpdateViewportRHI, we still provide the same inputs so we don't re-create the swapchain on a per-frame basis
+				ensure(PaintWindow == Window.Get());
 
 				// In the stereo case, the HMD display size drives the base RT size, separate from the PIE mirror window
 				const bool bResizeTargetValid = Window->IsViewportSizeDrivenByWindow() || 
@@ -1222,6 +1241,16 @@ FReply FSceneViewport::OnViewportActivated(const FWindowActivateEvent& InActivat
 	return FReply::Unhandled();
 }
 
+EDisplayColorGamut FSceneViewport::GetDisplayColorGamut() const
+{
+	return DisplayColorGamut;
+}
+
+EDisplayOutputFormat FSceneViewport::GetDisplayOutputFormat() const
+{
+	return DisplayOutputFormat;
+}
+
 void FSceneViewport::OnViewportDeactivated(const FWindowActivateEvent& InActivateEvent)
 {
 	// We backup if we have capture for us on activation, however we also maintain "true" if it's already true!
@@ -1662,6 +1691,22 @@ void FSceneViewport::UpdateViewportRHI(bool bDestroyed, uint32 NewSizeX, uint32
 
 		if( !bDestroyed )
 		{
+			TSharedPtr<SWidget> PinnedViewport = ViewportWidget.Pin();
+			void* OSWindow = nullptr;
+			FVector2D WindowTopLeft = FVector2D(0.0f, 0.0f);
+			FVector2D WindowBottomRight = FVector2D(0.0f, 0.0f);
+			if (PinnedViewport.IsValid())
+			{
+				TSharedPtr<SWindow> Window = FSlateApplication::Get().FindWidgetWindow(PinnedViewport.ToSharedRef());
+				if (Window.IsValid())
+				{
+					OSWindow = Window->GetNativeWindow()->GetOSWindowHandle();
+					WindowTopLeft = Window->GetPositionInScreen();
+					WindowBottomRight = Window->GetPositionInScreen() + Window->GetSizeInScreen();
+				}
+			}
+			HDRGetMetaData(DisplayOutputFormat, DisplayColorGamut, bHDRViewport, WindowTopLeft, WindowBottomRight, OSWindow);
+
 			BeginInitResource(this);
 			
 			FRenderCommandFence InitResourceFence;
@@ -1670,7 +1715,6 @@ void FSceneViewport::UpdateViewportRHI(bool bDestroyed, uint32 NewSizeX, uint32
 
 			FSlateRenderer* Renderer = FSlateApplication::Get().GetRenderer();
 
-			TSharedPtr<SWidget> PinnedViewport = ViewportWidget.Pin();
 			if (PinnedViewport.IsValid())
 			{
 				TSharedPtr<SWindow> Window = FSlateApplication::Get().FindWidgetWindow(PinnedViewport.ToSharedRef());
diff --git a/Engine/Source/Runtime/Engine/Private/UnrealEngine.cpp b/Engine/Source/Runtime/Engine/Private/UnrealEngine.cpp
--- a/Engine/Source/Runtime/Engine/Private/UnrealEngine.cpp
+++ b/Engine/Source/Runtime/Engine/Private/UnrealEngine.cpp
@@ -257,6 +257,7 @@ UnrealEngine.cpp: Implements the UEngine class and helpers.
 
 #include "Engine/InstancedStaticMesh.h"
 #include "IDeviceProfileSelectorModule.h"
+#include "HDRHelper.h"
 
 #if WITH_DUMPGPU
 	#include "RenderGraph.h"
@@ -868,7 +869,12 @@ void ENGINE_API HDRSettingChangedSinkCallback()
 	}
 
 	bool bIsHDREnabled = CVarHDROutputEnabled->GetValueOnAnyThread() != 0;
-	
+	if (bIsHDREnabled && !IsHDREnabled())
+	{
+		UE_LOG(LogConsoleResponse, Display, TEXT("Tried to enable HDR display output but unsupported, forcing off."));
+		bIsHDREnabled = false;
+	}
+
 	if(bIsHDREnabled != GRHIIsHDREnabled)
 	{
 		// We'll naively fall back to 1000 if DisplayNits is 0
@@ -884,26 +890,10 @@ void ENGINE_API HDRSettingChangedSinkCallback()
 			DisplayNitLevel = 1000;
 		}
 		
-		EDisplayOutputFormat OutputDevice = EDisplayOutputFormat::SDR_sRGB;
-		EDisplayColorGamut ColorGamut = EDisplayColorGamut::sRGB_D65;
-		
-		// If we are turning HDR on we must set the appropriate OutputDevice and ColorGamut.
-		// If we are turning it off, we'll reset back to 0/0
-		if(bIsHDREnabled)
-		{
-			FPlatformMisc::ChooseHDRDeviceAndColorGamut(GRHIVendorId, DisplayNitLevel, OutputDevice, ColorGamut);
-		}
-		
-		static IConsoleVariable* CVarHDROutputDevice = IConsoleManager::Get().FindConsoleVariable(TEXT("r.HDR.Display.OutputDevice"));
-		static IConsoleVariable* CVarHDRColorGamut = IConsoleManager::Get().FindConsoleVariable(TEXT("r.HDR.Display.ColorGamut"));
-		check(CVarHDROutputDevice);
-		check(CVarHDRColorGamut);
-		
-		CVarHDROutputDevice->Set((int32)OutputDevice, ECVF_SetByDeviceProfile);
-		CVarHDRColorGamut->Set((int32)ColorGamut, ECVF_SetByDeviceProfile);
-		
+		HDRConfigureCVars(bIsHDREnabled, DisplayNitLevel, false);
+
 		// Now set the HDR setting.
-		GRHIIsHDREnabled = CVarHDROutputEnabled->GetValueOnAnyThread() != 0;
+		GRHIIsHDREnabled = IsHDREnabled();
 	}
 }
 
diff --git a/Engine/Source/Runtime/Engine/Public/Slate/SceneViewport.h b/Engine/Source/Runtime/Engine/Public/Slate/SceneViewport.h
--- a/Engine/Source/Runtime/Engine/Public/Slate/SceneViewport.h
+++ b/Engine/Source/Runtime/Engine/Public/Slate/SceneViewport.h
@@ -242,7 +242,9 @@ public:
 	virtual FReply OnViewportActivated(const FWindowActivateEvent& InActivateEvent) override;
 	virtual void OnViewportDeactivated(const FWindowActivateEvent& InActivateEvent) override;
 	virtual FIntPoint GetSize() const override { return GetSizeXY(); }
-	
+	virtual EDisplayColorGamut GetDisplayColorGamut() const override;
+	virtual EDisplayOutputFormat GetDisplayOutputFormat() const override;
+
 	void SetViewportSize(uint32 NewSizeX,uint32 NewSizeY);
 	void SetFixedViewportSize(uint32 NewSizeX, uint32 NewSizeY);
 
@@ -444,6 +446,9 @@ private:
 	/** Tracks the number of touches currently active on the viewport */
 	int32 NumTouches;
 
+	EDisplayColorGamut DisplayColorGamut;
+	EDisplayOutputFormat DisplayOutputFormat;
+
 	/** The optional gamma value to use for this viewport */
 	TOptional<float> ViewportGammaOverride;
 };
diff --git a/Engine/Source/Runtime/Engine/Public/UnrealClient.h b/Engine/Source/Runtime/Engine/Public/UnrealClient.h
--- a/Engine/Source/Runtime/Engine/Public/UnrealClient.h
+++ b/Engine/Source/Runtime/Engine/Public/UnrealClient.h
@@ -60,6 +60,9 @@ public:
 	*/
 	ENGINE_API virtual float GetDisplayGamma() const;
 
+	virtual EDisplayColorGamut GetDisplayColorGamut() const { return EDisplayColorGamut::sRGB_D65; }
+	virtual EDisplayOutputFormat GetDisplayOutputFormat() const { return EDisplayOutputFormat::SDR_sRGB; }
+
 	/**
 	* Handles freezing/unfreezing of rendering
 	*/
diff --git a/Engine/Source/Runtime/MovieSceneCapture/Private/CompositionGraphCaptureProtocol.cpp b/Engine/Source/Runtime/MovieSceneCapture/Private/CompositionGraphCaptureProtocol.cpp
--- a/Engine/Source/Runtime/MovieSceneCapture/Private/CompositionGraphCaptureProtocol.cpp
+++ b/Engine/Source/Runtime/MovieSceneCapture/Private/CompositionGraphCaptureProtocol.cpp
@@ -7,12 +7,14 @@
 #include "Engine/Scene.h"
 #include "Materials/MaterialInterface.h"
 #include "SceneView.h"
+#include "Slate/SceneViewport.h"
 #include "Engine/Engine.h"
 #include "SceneViewExtension.h"
 #include "Materials/Material.h"
 #include "BufferVisualizationData.h"
 #include "MovieSceneCaptureModule.h"
 #include "MovieSceneCaptureSettings.h"
+#include "HDRHelper.h"
 
 struct FFrameCaptureViewExtension : public FSceneViewExtensionBase
 {
@@ -29,23 +31,9 @@ struct FFrameCaptureViewExtension : public FSceneViewExtensionBase
 		CVarDumpFrames = IConsoleManager::Get().FindConsoleVariable(TEXT("r.BufferVisualizationDumpFrames"));
 		CVarDumpFramesAsHDR = IConsoleManager::Get().FindConsoleVariable(TEXT("r.BufferVisualizationDumpFramesAsHDR"));
 		CVarHDRCompressionQuality = IConsoleManager::Get().FindConsoleVariable(TEXT("r.SaveEXR.CompressionQuality"));
-		CVarDumpGamut = IConsoleManager::Get().FindConsoleVariable(TEXT("r.HDR.Display.ColorGamut"));
-		CVarDumpDevice = IConsoleManager::Get().FindConsoleVariable(TEXT("r.HDR.Display.OutputDevice"));
 
 		RestoreDumpHDR = CVarDumpFramesAsHDR->GetInt();
 		RestoreHDRCompressionQuality = CVarHDRCompressionQuality->GetInt();
-		RestoreDumpGamut = (EDisplayColorGamut)CVarDumpGamut->GetInt();
-		RestoreDumpDevice = (EDisplayOutputFormat)CVarDumpDevice->GetInt();
-
-		if (CaptureGamut == HCGM_Linear)
-		{
-			CVarDumpGamut->Set((int32)EDisplayColorGamut::DCIP3_D65);
-			CVarDumpDevice->Set((int32)EDisplayOutputFormat::HDR_LinearEXR);
-		}
-		else
-		{
-			CVarDumpGamut->Set((int32)CaptureGamut);
-		}
 
 		Disable();
 	}
@@ -53,9 +41,6 @@ struct FFrameCaptureViewExtension : public FSceneViewExtensionBase
 	virtual ~FFrameCaptureViewExtension()
 	{
 		Disable();
-
-		CVarDumpGamut->Set((int32)RestoreDumpGamut);
-		CVarDumpDevice->Set((int32)RestoreDumpDevice);
 	}
 
 	bool IsEnabled() const
@@ -161,13 +146,9 @@ private:
 	IConsoleVariable* CVarDumpFrames;
 	IConsoleVariable* CVarDumpFramesAsHDR;
 	IConsoleVariable* CVarHDRCompressionQuality;
-	IConsoleVariable* CVarDumpGamut;
-	IConsoleVariable* CVarDumpDevice;
 
 	int32 RestoreDumpHDR;
 	int32 RestoreHDRCompressionQuality;
-	EDisplayColorGamut RestoreDumpGamut;
-	EDisplayOutputFormat RestoreDumpDevice;
 };
 
 bool UCompositionGraphCaptureProtocol::SetupImpl()
@@ -194,6 +175,23 @@ bool UCompositionGraphCaptureProtocol::SetupImpl()
 	PostProcessingMaterialPtr = Cast<UMaterialInterface>(PostProcessingMaterial.TryLoad());
 	ViewExtension = FSceneViewExtensions::NewExtension<FFrameCaptureViewExtension>(IncludeRenderPasses.Value, bCaptureFramesInHDR, HDRCompressionQuality, OverrideCaptureGamut, PostProcessingMaterialPtr, bDisableScreenPercentage);
 
+	EDisplayOutputFormat DisplayOutputFormat = HDRGetDefaultDisplayOutputFormat();
+	EDisplayColorGamut DisplayColorGamut = HDRGetDefaultDisplayColorGamut();
+	bool bHDREnabled = IsHDREnabled() && GRHISupportsHDROutput;
+
+	if (CaptureGamut == HCGM_Linear)
+	{
+		DisplayColorGamut = EDisplayColorGamut::DCIP3_D65;
+		DisplayOutputFormat = EDisplayOutputFormat::HDR_LinearEXR;
+	}
+	else
+	{
+		DisplayColorGamut = (EDisplayColorGamut)CaptureGamut.GetValue();
+	}
+
+	TSharedPtr<SWindow> CustomWindow = InitSettings->SceneViewport->FindWindow();
+	HDRAddCustomMetaData(CustomWindow->GetNativeWindow()->GetOSWindowHandle(), DisplayOutputFormat, DisplayColorGamut, bHDREnabled);
+
 	return true;
 }
 
@@ -240,6 +238,9 @@ void UCompositionGraphCaptureProtocol::OnLoadConfigImpl(FMovieSceneCaptureSettin
 
 void UCompositionGraphCaptureProtocol::FinalizeImpl()
 {
+	TSharedPtr<SWindow> CustomWindow = InitSettings->SceneViewport->FindWindow();
+	HDRRemoveCustomMetaData(CustomWindow->GetNativeWindow()->GetOSWindowHandle());
+
 	ViewExtension->Disable(true);
 }
 
diff --git a/Engine/Source/Runtime/MovieSceneCapture/Private/ImageSequenceProtocol.cpp b/Engine/Source/Runtime/MovieSceneCapture/Private/ImageSequenceProtocol.cpp
--- a/Engine/Source/Runtime/MovieSceneCapture/Private/ImageSequenceProtocol.cpp
+++ b/Engine/Source/Runtime/MovieSceneCapture/Private/ImageSequenceProtocol.cpp
@@ -16,6 +16,7 @@
 #include "MovieSceneCaptureModule.h"
 #include "MovieSceneCaptureSettings.h"
 #include "ImageWriteQueue.h"
+#include "HDRHelper.h"
 
 struct FImageFrameData : IFramePayload
 {
@@ -206,34 +207,32 @@ bool UImageSequenceProtocol_EXR::SetupImpl()
 		bCompressed = HDRCompressionQuality != (int32)EImageCompressionQuality::Uncompressed;
 	}
 
-	IConsoleVariable* CVarDumpGamut = IConsoleManager::Get().FindConsoleVariable(TEXT("r.HDR.Display.ColorGamut"));
-	IConsoleVariable* CVarDumpDevice = IConsoleManager::Get().FindConsoleVariable(TEXT("r.HDR.Display.OutputDevice"));
-
-	RestoreColorGamut = (EDisplayColorGamut)CVarDumpGamut->GetInt();
-	RestoreOutputDevice = (EDisplayOutputFormat)CVarDumpDevice->GetInt();
+	EDisplayOutputFormat DisplayOutputFormat = HDRGetDefaultDisplayOutputFormat();
+	EDisplayColorGamut DisplayColorGamut = HDRGetDefaultDisplayColorGamut();
+	bool bHDREnabled = IsHDREnabled() && GRHISupportsHDROutput;
 
 	if (CaptureGamut == HCGM_Linear)
 	{
-		CVarDumpGamut->Set((int32)EDisplayColorGamut::DCIP3_D65);
-		CVarDumpDevice->Set((int32)EDisplayOutputFormat::HDR_LinearEXR);
+		DisplayColorGamut = EDisplayColorGamut::DCIP3_D65;
+		DisplayOutputFormat = EDisplayOutputFormat::HDR_LinearEXR;
 	}
 	else
 	{
-		CVarDumpGamut->Set(CaptureGamut);
+		DisplayColorGamut = (EDisplayColorGamut)CaptureGamut.GetValue();
 	}
 
+	TSharedPtr<SWindow> CustomWindow = InitSettings->SceneViewport->FindWindow();
+	HDRAddCustomMetaData(CustomWindow->GetNativeWindow()->GetOSWindowHandle(), DisplayOutputFormat, DisplayColorGamut, bHDREnabled);
+
 	return Super::SetupImpl();
 }
 
 void UImageSequenceProtocol_EXR::FinalizeImpl()
 {
-	Super::FinalizeImpl();
+	TSharedPtr<SWindow> CustomWindow = InitSettings->SceneViewport->FindWindow();
+	HDRRemoveCustomMetaData(CustomWindow->GetNativeWindow()->GetOSWindowHandle());
 
-	IConsoleVariable* CVarDumpGamut = IConsoleManager::Get().FindConsoleVariable(TEXT("r.HDR.Display.ColorGamut"));
-	IConsoleVariable* CVarDumpDevice = IConsoleManager::Get().FindConsoleVariable(TEXT("r.HDR.Display.OutputDevice"));
-
-	CVarDumpGamut->Set((int32)RestoreColorGamut);
-	CVarDumpDevice->Set((int32)RestoreOutputDevice);
+	Super::FinalizeImpl();
 }
 
 void UImageSequenceProtocol_EXR::AddFormatMappingsImpl(TMap<FString, FStringFormatArg>& FormatMappings) const
diff --git a/Engine/Source/Runtime/RHI/Public/DynamicRHI.h b/Engine/Source/Runtime/RHI/Public/DynamicRHI.h
--- a/Engine/Source/Runtime/RHI/Public/DynamicRHI.h
+++ b/Engine/Source/Runtime/RHI/Public/DynamicRHI.h
@@ -849,6 +849,10 @@ public:
 		return nullptr;
 	}
 
+	virtual void RHIGetDisplaysInformation(FDisplayInformationArray& OutDisplayInformation)
+	{
+	}
+
 	virtual void RHIAdvanceFrameFence() {};
 
 	// Only relevant with an RHI thread, this advances the backbuffer for the purpose of GetViewportBackBuffer
@@ -1627,6 +1631,10 @@ FORCEINLINE IRHITransientResourceAllocator* RHICreateTransientResourceAllocator(
 	return GDynamicRHI->RHICreateTransientResourceAllocator();
 }
 
+FORCEINLINE void RHIGetDisplaysInformation(FDisplayInformationArray& OutDisplayInformation)
+{
+	GDynamicRHI->RHIGetDisplaysInformation(OutDisplayInformation);
+}
 
 #if RHI_RAYTRACING
 
diff --git a/Engine/Source/Runtime/RHI/Public/RHIDefinitions.h b/Engine/Source/Runtime/RHI/Public/RHIDefinitions.h
--- a/Engine/Source/Runtime/RHI/Public/RHIDefinitions.h
+++ b/Engine/Source/Runtime/RHI/Public/RHIDefinitions.h
@@ -2033,6 +2033,8 @@ private:
 	ERHIDescriptorHeapType Type{ ERHIDescriptorHeapType::count };
 };
 
+using FDisplayInformationArray = TArray<struct FDisplayInformation>;
+
 inline bool IsPCPlatform(const FStaticShaderPlatform Platform)
 {
 	return FDataDrivenShaderPlatformInfo::GetIsPC(Platform);
diff --git a/Engine/Source/Runtime/RenderCore/Private/RenderCore.cpp b/Engine/Source/Runtime/RenderCore/Private/RenderCore.cpp
--- a/Engine/Source/Runtime/RenderCore/Private/RenderCore.cpp
+++ b/Engine/Source/Runtime/RenderCore/Private/RenderCore.cpp
@@ -9,6 +9,7 @@
 #include "UniformBuffer.h"
 #include "Modules/ModuleManager.h"
 #include "Shader.h"
+#include "HDRHelper.h"
 
 void UpdateShaderDevelopmentMode();
 
@@ -305,17 +306,65 @@ RENDERCORE_API int32 GetCVarForceLODShadow_AnyThread()
 	return Ret;
 }
 
+FMatrix44f FVirtualTextureUniformData::Invalid = FMatrix44f::Identity;
+
+// Note: Enables or disables HDR support for a project. Typically this would be set on a per-project/per-platform basis in defaultengine.ini
+TAutoConsoleVariable<int32> CVarAllowHDR(
+	TEXT("r.AllowHDR"),
+	0,
+	TEXT("Creates an HDR compatible swap-chain and enables HDR display output.")
+	TEXT("0: Disabled (default)\n")
+	TEXT("1: Allow HDR, if supported by the platform and display \n"),
+	ECVF_ReadOnly);
+
+// Note: These values are directly referenced in code. They are set in code at runtime and therefore cannot be set via ini files
+// Please update all paths if changing
+TAutoConsoleVariable<int32> CVarDisplayColorGamut(
+	TEXT("r.HDR.Display.ColorGamut"),
+	0,
+	TEXT("Color gamut of the output display:\n")
+	TEXT("0: Rec709 / sRGB, D65 (default)\n")
+	TEXT("1: DCI-P3, D65\n")
+	TEXT("2: Rec2020 / BT2020, D65\n")
+	TEXT("3: ACES, D60\n")
+	TEXT("4: ACEScg, D60\n"),
+	ECVF_Scalability | ECVF_RenderThreadSafe);
+
+TAutoConsoleVariable<int32> CVarDisplayOutputDevice(
+	TEXT("r.HDR.Display.OutputDevice"),
+	0,
+	TEXT("Device format of the output display:\n")
+	TEXT("0: sRGB (LDR)\n")
+	TEXT("1: Rec709 (LDR)\n")
+	TEXT("2: Explicit gamma mapping (LDR)\n")
+	TEXT("3: ACES 1000 nit ST-2084 (Dolby PQ) (HDR)\n")
+	TEXT("4: ACES 2000 nit ST-2084 (Dolby PQ) (HDR)\n")
+	TEXT("5: ACES 1000 nit ScRGB (HDR)\n")
+	TEXT("6: ACES 2000 nit ScRGB (HDR)\n")
+	TEXT("7: Linear EXR (HDR)\n")
+	TEXT("8: Linear final color, no tone curve (HDR)\n")
+	TEXT("9: Linear final color with tone curve\n"),
+	ECVF_Scalability | ECVF_RenderThreadSafe);
+
+static TAutoConsoleVariable<int32> CVarHDRNits(
+	TEXT("r.HDR.DisplayNitsLevel"),
+	0,
+	TEXT("The configured display output nit level, assuming HDR output is enabled."),
+	ECVF_RenderThreadSafe
+);
+
+TAutoConsoleVariable<int32> CVarHDROutputEnabled(
+	TEXT("r.HDR.EnableHDROutput"),
+	0,
+	TEXT("Creates an HDR compatible swap-chain and enables HDR display output.")
+	TEXT("0: Disabled (default)\n")
+	TEXT("1: Enable hardware-specific implementation\n"),
+	ECVF_RenderThreadSafe
+);
+
 RENDERCORE_API bool IsHDREnabled()
 {
-	static const auto CVarHDROutputEnabled = IConsoleManager::Get().FindTConsoleVariableDataInt(TEXT("r.HDR.EnableHDROutput"));
-	if (CVarHDROutputEnabled)
-	{
-		if (CVarHDROutputEnabled->GetValueOnAnyThread() != 0)
-		{
-			return true;
-		}
-	}
-	return false;
+	return GRHISupportsHDROutput && CVarHDROutputEnabled.GetValueOnAnyThread() != 0;
 }
 
 RENDERCORE_API bool IsHDRAllowed()
@@ -330,13 +379,178 @@ RENDERCORE_API bool IsHDRAllowed()
 		return false;
 	}
 
-	static const auto CVarHDRAllow = IConsoleManager::Get().FindTConsoleVariableDataInt(TEXT("r.AllowHDR"));
-	if (CVarHDRAllow && CVarHDRAllow->GetValueOnAnyThread() != 0)
+	return (CVarAllowHDR.GetValueOnAnyThread() != 0);
+}
+
+RENDERCORE_API EDisplayOutputFormat HDRGetDefaultDisplayOutputFormat()
+{
+	return static_cast<EDisplayOutputFormat>(FMath::Clamp(CVarDisplayOutputDevice.GetValueOnAnyThread(), 0, static_cast<int32>(EDisplayOutputFormat::MAX) - 1));
+}
+
+RENDERCORE_API EDisplayColorGamut HDRGetDefaultDisplayColorGamut()
+{
+	return static_cast<EDisplayColorGamut>(FMath::Clamp(CVarDisplayColorGamut.GetValueOnAnyThread(), 0, static_cast<int32>(EDisplayColorGamut::MAX) - 1));
+}
+
+struct FHDRMetaData
+{
+	EDisplayOutputFormat DisplayOutputFormat;
+	EDisplayColorGamut DisplayColorGamut;
+	uint32 MaximumLuminanceInNits;
+	bool bHDRSupported;
+};
+
+
+inline FHDRMetaData HDRGetDefaultMetaData()
+{
+	FHDRMetaData HDRMetaData{};
+	HDRMetaData.DisplayOutputFormat = HDRGetDefaultDisplayOutputFormat();
+	HDRMetaData.DisplayColorGamut = HDRGetDefaultDisplayColorGamut();
+	HDRMetaData.bHDRSupported = IsHDREnabled() && GRHISupportsHDROutput;
+	HDRMetaData.MaximumLuminanceInNits = CVarHDRNits.GetValueOnAnyThread();
+	return HDRMetaData;
+}
+
+inline int32 WindowDisplayIntersectionArea(FIntRect WindowRect, FIntRect DisplayRect)
+{
+	return FMath::Max<int32>(0, FMath::Min<int32>(WindowRect.Max.X, DisplayRect.Max.X) - FMath::Max<int32>(WindowRect.Min.X, DisplayRect.Min.X)) *
+		FMath::Max<int32>(0, FMath::Min<int32>(WindowRect.Max.Y, DisplayRect.Max.Y) - FMath::Max<int32>(WindowRect.Min.Y, DisplayRect.Min.Y));
+}
+
+TMap<void*, FHDRMetaData> GWindowsWithDefaultParams;
+FCriticalSection GWindowsWithDefaultParamsCS;
+
+RENDERCORE_API void HDRAddCustomMetaData(void* OSWindow, EDisplayOutputFormat DisplayOutputFormat, EDisplayColorGamut DisplayColorGamut, bool bHDREnabled)
+{
+	ensure(OSWindow != nullptr);
+	if (OSWindow == nullptr)
+	{
+		return;
+	}
+	FHDRMetaData HDRMetaData{};
+	HDRMetaData.DisplayOutputFormat = DisplayOutputFormat;
+	HDRMetaData.DisplayColorGamut = DisplayColorGamut;
+	HDRMetaData.bHDRSupported = bHDREnabled;
+	HDRMetaData.MaximumLuminanceInNits = CVarHDRNits.GetValueOnAnyThread();
+
+	FScopeLock Lock(&GWindowsWithDefaultParamsCS);
+	GWindowsWithDefaultParams.Add(OSWindow, HDRMetaData);
+}
+
+RENDERCORE_API void HDRRemoveCustomMetaData(void* OSWindow)
+{
+	ensure(OSWindow != nullptr);
+	if (OSWindow == nullptr)
+	{
+		return;
+	}
+	FScopeLock Lock(&GWindowsWithDefaultParamsCS);
+	GWindowsWithDefaultParams.Remove(OSWindow);
+}
+
+bool HdrHasWindowParamsFromCVars(void* OSWindow, FHDRMetaData& HDRMetaData)
+{
+	if (GWindowsWithDefaultParams.Num() == 0)
 	{
+		return false;
+	}
+	FScopeLock Lock(&GWindowsWithDefaultParamsCS);
+	FHDRMetaData* FoundHDRMetaData = GWindowsWithDefaultParams.Find(OSWindow);
+	if (FoundHDRMetaData)
+	{
+		HDRMetaData = *FoundHDRMetaData;
 		return true;
 	}
 
 	return false;
 }
 
-FMatrix44f FVirtualTextureUniformData::Invalid = FMatrix44f::Identity;
\ No newline at end of file
+RENDERCORE_API void HDRGetMetaData(EDisplayOutputFormat& OutDisplayOutputFormat, EDisplayColorGamut& OutDisplayColorGamut, bool& OutbHDRSupported, 
+								   const FVector2D& WindowTopLeft, const FVector2D& WindowBottomRight, void* OSWindow)
+{
+	FHDRMetaData HDRMetaData;
+
+#if WITH_EDITOR
+	// this has priority over IsHDREnabled because MovieSceneCapture might request custom parameters
+	if (HdrHasWindowParamsFromCVars(OSWindow, HDRMetaData))
+	{
+		return;
+	}
+#endif
+	
+	HDRMetaData = HDRGetDefaultMetaData();
+
+	OutDisplayOutputFormat = HDRMetaData.DisplayOutputFormat;
+	OutDisplayColorGamut = HDRMetaData.DisplayColorGamut;
+	OutbHDRSupported = HDRMetaData.bHDRSupported;
+	if (!IsHDREnabled() || OSWindow == nullptr)
+	{
+		return;
+	}
+
+	FDisplayInformationArray DisplayList;
+	RHIGetDisplaysInformation(DisplayList);
+	// In case we have 1 display or less, the CVars that were setup do represent the state of the displays
+	if (DisplayList.Num() <= 1)
+	{
+		return;
+	}
+
+	FIntRect WindowRect((int32)WindowTopLeft.X, (int32)WindowTopLeft.Y, (int32)WindowBottomRight.X, (int32)WindowBottomRight.Y);
+	int32 BestDisplayForWindow = 0;
+	int32 BestArea = 0;
+	for (int32 DisplayIndex = 0; DisplayIndex < DisplayList.Num(); ++DisplayIndex)
+	{
+		// Compute the intersection
+		int32 CurrentArea = WindowDisplayIntersectionArea(WindowRect, DisplayList[DisplayIndex].DesktopCoordinates);
+		if (CurrentArea > BestArea)
+		{
+			BestDisplayForWindow = DisplayIndex;
+			BestArea = CurrentArea;
+		}
+	}
+
+	OutbHDRSupported = DisplayList[BestDisplayForWindow].bHDRSupported;
+	OutDisplayOutputFormat = EDisplayOutputFormat::SDR_sRGB;
+	OutDisplayColorGamut = EDisplayColorGamut::sRGB_D65;
+
+	if (OutbHDRSupported)
+	{
+		FPlatformMisc::ChooseHDRDeviceAndColorGamut(GRHIVendorId, CVarHDRNits.GetValueOnAnyThread(), OutDisplayOutputFormat, OutDisplayColorGamut);
+	}
+
+}
+
+RENDERCORE_API void HDRConfigureCVars(bool bIsHDREnabled, uint32 DisplayNits, bool bFromGameSettings)
+{
+	if (bIsHDREnabled && !GRHISupportsHDROutput)
+	{
+		UE_LOG(LogRendererCore, Warning, TEXT("Trying to enable HDR but it is not supported by the RHI: IsHDREnabled will return false"));
+		bIsHDREnabled = false;
+	}
+
+	EDisplayOutputFormat OutputDevice = EDisplayOutputFormat::SDR_sRGB;
+	EDisplayColorGamut ColorGamut = EDisplayColorGamut::sRGB_D65;
+
+	// If we are turning HDR on we must set the appropriate OutputDevice and ColorGamut.
+	// If we are turning it off, we'll reset back to 0/0
+	if (bIsHDREnabled)
+	{
+		FPlatformMisc::ChooseHDRDeviceAndColorGamut(GRHIVendorId, DisplayNits, OutputDevice, ColorGamut);
+	}
+
+	//CVarHDRNits is ECVF_SetByCode as it's only a mean of communicating the information from UGameUserSettings to the rest of the engine
+	if (bIsHDREnabled)
+	{
+		CVarHDROutputEnabled->Set(1, bFromGameSettings ? ECVF_SetByGameSetting : ECVF_SetByCode);
+		CVarHDRNits->Set((int32)DisplayNits, ECVF_SetByCode);
+	}
+	else
+	{
+		CVarHDROutputEnabled->Set(0, bFromGameSettings ? ECVF_SetByGameSetting : ECVF_SetByCode);
+		CVarHDRNits->Set(0, ECVF_SetByCode);
+	}
+
+	CVarDisplayOutputDevice->Set((int32)OutputDevice, ECVF_SetByDeviceProfile);
+	CVarDisplayColorGamut->Set((int32)ColorGamut, ECVF_SetByDeviceProfile);
+}
\ No newline at end of file
diff --git a/Engine/Source/Runtime/RenderCore/Public/HDRHelper.h b/Engine/Source/Runtime/RenderCore/Public/HDRHelper.h
new file mode 100644
--- /dev/null
+++ b/Engine/Source/Runtime/RenderCore/Public/HDRHelper.h
@@ -0,0 +1,26 @@
+// Copyright Epic Games, Inc. All Rights Reserved.
+
+#pragma once
+
+#include "CoreMinimal.h"
+
+RENDERCORE_API void HDRGetMetaData(EDisplayOutputFormat& OutDisplayOutputFormat, EDisplayColorGamut& OutDisplayColorGamut, bool& OutbHDRSupported,
+								   const FVector2D& WindowTopLeft, const FVector2D& WindowBottomRight, void* OSWindow);
+
+RENDERCORE_API void HDRConfigureCVars(bool bIsHDREnabled, uint32 DisplayNits, bool bFromGameSettings);
+RENDERCORE_API EDisplayOutputFormat HDRGetDefaultDisplayOutputFormat();
+RENDERCORE_API EDisplayColorGamut HDRGetDefaultDisplayColorGamut();
+RENDERCORE_API void HDRAddCustomMetaData(void* OSWindow, EDisplayOutputFormat DisplayOutputFormat, EDisplayColorGamut DisplayColorGamut, bool bHDREnabled);
+RENDERCORE_API void HDRRemoveCustomMetaData(void* OSWindow);
+
+
+struct FDisplayInformation
+{
+	FDisplayInformation()
+		: DesktopCoordinates(0, 0, 0, 0)
+		, bHDRSupported(false)
+	{}
+
+	FIntRect DesktopCoordinates;
+	bool   bHDRSupported;
+};
diff --git a/Engine/Source/Runtime/Renderer/Private/PostProcess/PostProcessDeviceEncodingOnly.cpp b/Engine/Source/Runtime/Renderer/Private/PostProcess/PostProcessDeviceEncodingOnly.cpp
--- a/Engine/Source/Runtime/Renderer/Private/PostProcess/PostProcessDeviceEncodingOnly.cpp
+++ b/Engine/Source/Runtime/Renderer/Private/PostProcess/PostProcessDeviceEncodingOnly.cpp
@@ -34,8 +34,6 @@ namespace
 
 FDeviceEncodingOnlyOutputDeviceParameters GetDeviceEncodingOnlyOutputDeviceParameters(const FSceneViewFamily& Family)
 {
-	static TConsoleVariableData<int32>* CVarOutputGamut = IConsoleManager::Get().FindTConsoleVariableDataInt(TEXT("r.HDR.Display.ColorGamut"));
-	static TConsoleVariableData<int32>* CVarOutputDevice = IConsoleManager::Get().FindTConsoleVariableDataInt(TEXT("r.HDR.Display.OutputDevice"));
 	static TConsoleVariableData<float>* CVarOutputGamma = IConsoleManager::Get().FindTConsoleVariableDataFloat(TEXT("r.TonemapperGamma"));
 
 	EDisplayOutputFormat OutputDeviceValue;
@@ -54,8 +52,7 @@ FDeviceEncodingOnlyOutputDeviceParameters GetDeviceEncodingOnlyOutputDeviceParam
 	}
 	else
 	{
-		OutputDeviceValue = static_cast<EDisplayOutputFormat>(CVarOutputDevice->GetValueOnRenderThread());
-		OutputDeviceValue = static_cast<EDisplayOutputFormat>(FMath::Clamp(static_cast<int32>(OutputDeviceValue), 0, static_cast<int32>(EDisplayOutputFormat::MAX) - 1));
+		OutputDeviceValue = Family.RenderTarget->GetDisplayOutputFormat();
 	}
 
 	float Gamma = CVarOutputGamma->GetValueOnRenderThread();
@@ -79,7 +76,7 @@ FDeviceEncodingOnlyOutputDeviceParameters GetDeviceEncodingOnlyOutputDeviceParam
 	FDeviceEncodingOnlyOutputDeviceParameters Parameters;
 	Parameters.InverseGamma = (FVector3f)InvDisplayGammaValue;
 	Parameters.OutputDevice = static_cast<uint32>(OutputDeviceValue);
-	Parameters.OutputGamut = CVarOutputGamut->GetValueOnRenderThread();
+	Parameters.OutputGamut = static_cast<uint32>(Family.RenderTarget->GetDisplayColorGamut());
 	return Parameters;
 }
 
diff --git a/Engine/Source/Runtime/Renderer/Private/PostProcess/PostProcessTonemap.cpp b/Engine/Source/Runtime/Renderer/Private/PostProcess/PostProcessTonemap.cpp
--- a/Engine/Source/Runtime/Renderer/Private/PostProcess/PostProcessTonemap.cpp
+++ b/Engine/Source/Runtime/Renderer/Private/PostProcess/PostProcessTonemap.cpp
@@ -35,53 +35,6 @@ TAutoConsoleVariable<float> CVarTonemapperSharpen(
 	TEXT("   1: full strength"),
 	ECVF_Scalability | ECVF_RenderThreadSafe);
 
-// Note: Enables or disables HDR support for a project. Typically this would be set on a per-project/per-platform basis in defaultengine.ini
-TAutoConsoleVariable<int32> CVarAllowHDR(
-	TEXT("r.AllowHDR"),
-	0,
-	TEXT("Creates an HDR compatible swap-chain and enables HDR display output.")
-	TEXT("0: Disabled (default)\n")
-	TEXT("1: Allow HDR, if supported by the platform and display \n"),
-	ECVF_ReadOnly);
-
-// Note: These values are directly referenced in code. They are set in code at runtime and therefore cannot be set via ini files
-// Please update all paths if changing
-TAutoConsoleVariable<int32> CVarDisplayColorGamut(
-	TEXT("r.HDR.Display.ColorGamut"),
-	0,
-	TEXT("Color gamut of the output display:\n")
-	TEXT("0: Rec709 / sRGB, D65 (default)\n")
-	TEXT("1: DCI-P3, D65\n")
-	TEXT("2: Rec2020 / BT2020, D65\n")
-	TEXT("3: ACES, D60\n")
-	TEXT("4: ACEScg, D60\n"),
-	ECVF_Scalability | ECVF_RenderThreadSafe);
-
-TAutoConsoleVariable<int32> CVarDisplayOutputDevice(
-	TEXT("r.HDR.Display.OutputDevice"),
-	0,
-	TEXT("Device format of the output display:\n")
-	TEXT("0: sRGB (LDR)\n")
-	TEXT("1: Rec709 (LDR)\n")
-	TEXT("2: Explicit gamma mapping (LDR)\n")
-	TEXT("3: ACES 1000 nit ST-2084 (Dolby PQ) (HDR)\n")
-	TEXT("4: ACES 2000 nit ST-2084 (Dolby PQ) (HDR)\n")
-	TEXT("5: ACES 1000 nit ScRGB (HDR)\n")
-	TEXT("6: ACES 2000 nit ScRGB (HDR)\n")
-	TEXT("7: Linear EXR (HDR)\n")
-	TEXT("8: Linear final color, no tone curve (HDR)\n")
-	TEXT("9: Linear final color with tone curve\n"),
-	ECVF_Scalability | ECVF_RenderThreadSafe);
-	
-TAutoConsoleVariable<int32> CVarHDROutputEnabled(
-	TEXT("r.HDR.EnableHDROutput"),
-	0,
-	TEXT("Creates an HDR compatible swap-chain and enables HDR display output.")
-	TEXT("0: Disabled (default)\n")
-	TEXT("1: Enable hardware-specific implementation\n"),
-	ECVF_RenderThreadSafe
-	);
-
 TAutoConsoleVariable<float> CVarTonemapperGamma(
 	TEXT("r.TonemapperGamma"),
 	0.0f,
@@ -271,8 +224,6 @@ bool ShouldCompileDesktopPermutation(const FGlobalShaderPermutationParameters& P
 
 FTonemapperOutputDeviceParameters GetTonemapperOutputDeviceParameters(const FSceneViewFamily& Family)
 {
-	static TConsoleVariableData<int32>* CVarOutputGamut = IConsoleManager::Get().FindTConsoleVariableDataInt(TEXT("r.HDR.Display.ColorGamut"));
-	static TConsoleVariableData<int32>* CVarOutputDevice = IConsoleManager::Get().FindTConsoleVariableDataInt(TEXT("r.HDR.Display.OutputDevice"));
 	static TConsoleVariableData<float>* CVarOutputGamma = IConsoleManager::Get().FindTConsoleVariableDataFloat(TEXT("r.TonemapperGamma"));
 
 	EDisplayOutputFormat OutputDeviceValue;
@@ -291,8 +242,7 @@ FTonemapperOutputDeviceParameters GetTonemapperOutputDeviceParameters(const FSce
 	}
 	else
 	{
-		OutputDeviceValue = static_cast<EDisplayOutputFormat>(CVarOutputDevice->GetValueOnRenderThread());
-		OutputDeviceValue = static_cast<EDisplayOutputFormat>(FMath::Clamp(static_cast<int32>(OutputDeviceValue), 0, static_cast<int32>(EDisplayOutputFormat::MAX) - 1));
+		OutputDeviceValue = Family.RenderTarget->GetDisplayOutputFormat();
 	}
 
 	float Gamma = CVarOutputGamma->GetValueOnRenderThread();
@@ -316,7 +266,7 @@ FTonemapperOutputDeviceParameters GetTonemapperOutputDeviceParameters(const FSce
 	FTonemapperOutputDeviceParameters Parameters;
 	Parameters.InverseGamma = InvDisplayGammaValue;
 	Parameters.OutputDevice = static_cast<uint32>(OutputDeviceValue);
-	Parameters.OutputGamut = CVarOutputGamut->GetValueOnRenderThread();
+	Parameters.OutputGamut = static_cast<uint32>(Family.RenderTarget->GetDisplayColorGamut());
 	return Parameters;
 }
 
diff --git a/Engine/Source/Runtime/SlateRHIRenderer/Private/SlateRHIRenderer.cpp b/Engine/Source/Runtime/SlateRHIRenderer/Private/SlateRHIRenderer.cpp
--- a/Engine/Source/Runtime/SlateRHIRenderer/Private/SlateRHIRenderer.cpp
+++ b/Engine/Source/Runtime/SlateRHIRenderer/Private/SlateRHIRenderer.cpp
@@ -34,6 +34,7 @@
 #include "HAL/LowLevelMemTracker.h"
 #include "Rendering/RenderingCommon.h"
 #include "IHeadMountedDisplayModule.h"
+#include "HDRHelper.h"
 
 DECLARE_CYCLE_STAT(TEXT("Slate RT: Rendering"), STAT_SlateRenderingRTTime, STATGROUP_Slate);
 
@@ -377,9 +378,9 @@ void FSlateRHIRenderer::CreateViewport(const TSharedRef<SWindow> Window)
 
 		// SDR format holds the requested format in non HDR mode
 		NewInfo->SDRPixelFormat = NewInfo->PixelFormat;
-		bool HDREnabled = IsHDREnabled();
-		NewInfo->bSceneHDREnabled = HDREnabled;
-		if (HDREnabled)
+		HDRGetMetaData(NewInfo->HDRDisplayOutputFormat, NewInfo->HDRDisplayColorGamut, NewInfo->bSceneHDREnabled, Window->GetPositionInScreen(), Window->GetPositionInScreen() + Window->GetSizeInScreen(), NewInfo->OSWindow);
+
+		if (NewInfo->bSceneHDREnabled)
 		{
 			NewInfo->PixelFormat = GRHIHDRDisplayOutputFormat;
 		}
@@ -401,30 +402,23 @@ void FSlateRHIRenderer::CreateViewport(const TSharedRef<SWindow> Window)
 	}
 }
 
-void FSlateRHIRenderer::ConditionalResizeViewport(FViewportInfo* ViewInfo, uint32 Width, uint32 Height, bool bFullscreen)
+void FSlateRHIRenderer::ConditionalResizeViewport(FViewportInfo* ViewInfo, uint32 Width, uint32 Height, bool bFullscreen, SWindow* Window)
 {
 	checkSlow(IsThreadSafeForSlateRendering());
 
 	// Force update if HDR output state changes
-	static const auto CVarHDRColorGamut = IConsoleManager::Get().FindTConsoleVariableDataInt(TEXT("r.HDR.Display.ColorGamut"));
-	static const auto CVarHDROutputDevice = IConsoleManager::Get().FindTConsoleVariableDataInt(TEXT("r.HDR.Display.OutputDevice"));
 
 	bool bHDREnabled = IsHDREnabled();
-	int32 HDRColorGamut = CVarHDRColorGamut ? CVarHDRColorGamut->GetValueOnAnyThread() : 0;
-	int32 HDROutputDevice = CVarHDROutputDevice ? CVarHDROutputDevice->GetValueOnAnyThread() : 0;
-
-	bool bHDRStale = GRHISupportsHDROutput && ViewInfo && (
-		((ViewInfo->PixelFormat == GRHIHDRDisplayOutputFormat) != bHDREnabled)	// HDR toggled
-#if PLATFORM_WINDOWS
-		|| ((IsRHIDeviceNVIDIA() || IsRHIDeviceAMD()) &&						// Vendor-specific mastering data updates
-		((bHDREnabled && ViewInfo->HDRColorGamut != HDRColorGamut)				// Color gamut changed
-			|| (bHDREnabled && ViewInfo->HDROutputDevice != HDROutputDevice)))	// Output device changed
-#endif
-		);
+	EDisplayColorGamut HDRColorGamut = HDRGetDefaultDisplayColorGamut();
+	EDisplayOutputFormat HDROutputDevice = HDRGetDefaultDisplayOutputFormat();
 
-	if (ViewInfo && (ViewInfo->bSceneHDREnabled != bHDREnabled))
+	bool bHDRStale = false;
+	if (ViewInfo)
 	{
-		bHDRStale = true;
+		HDRGetMetaData(HDROutputDevice, HDRColorGamut, bHDREnabled, Window->GetPositionInScreen(), Window->GetPositionInScreen() + Window->GetSizeInScreen(), ViewInfo->OSWindow);
+		bHDRStale |= HDROutputDevice != ViewInfo->HDRDisplayOutputFormat;
+		bHDRStale |= HDRColorGamut != ViewInfo->HDRDisplayColorGamut;
+		bHDRStale |= bHDREnabled != ViewInfo->bSceneHDREnabled;
 	}
 
 	if (IsInGameThread() && !IsInSlateThread() && ViewInfo && (bHDRStale || ViewInfo->Height != Height || ViewInfo->Width != Width || ViewInfo->bFullscreen != bFullscreen || !IsValidRef(ViewInfo->ViewportRHI)))
@@ -465,8 +459,8 @@ void FSlateRHIRenderer::ConditionalResizeViewport(FViewportInfo* ViewInfo, uint3
 		ViewInfo->bFullscreen = bFullscreen;
 
 		ViewInfo->PixelFormat = bHDREnabled ? GRHIHDRDisplayOutputFormat : ViewInfo->SDRPixelFormat;
-		ViewInfo->HDRColorGamut = HDRColorGamut;
-		ViewInfo->HDROutputDevice = HDROutputDevice;
+		ViewInfo->HDRDisplayColorGamut = HDRColorGamut;
+		ViewInfo->HDRDisplayOutputFormat = HDROutputDevice;
 		ViewInfo->bSceneHDREnabled = bHDREnabled;
 
 		PreResizeBackBufferDelegate.Broadcast(&ViewInfo->ViewportRHI);
@@ -512,7 +506,7 @@ void FSlateRHIRenderer::UpdateFullscreenState(const TSharedRef<SWindow> Window,
 			ResY = ViewInfo->DesiredHeight;
 		}
 
-		ConditionalResizeViewport(ViewInfo, ResX, ResY, bFullscreen);
+		ConditionalResizeViewport(ViewInfo, ResX, ResY, bFullscreen, &Window.Get());
 	}
 }
 
@@ -591,14 +585,12 @@ public:
 	}
 	FCompositeLUTGenerationPS() {}
 
-	void SetParameters(FRHICommandList& RHICmdList)
+	void SetParameters(FRHICommandList& RHICmdList, EDisplayOutputFormat DisplayOutputFormat, EDisplayColorGamut DisplayColorGamut)
 	{
-		static const auto CVarOutputDevice = IConsoleManager::Get().FindTConsoleVariableDataInt(TEXT("r.HDR.Display.OutputDevice"));
-		static const auto CVarOutputGamut = IConsoleManager::Get().FindTConsoleVariableDataInt(TEXT("r.HDR.Display.ColorGamut"));
 		static const auto CVarOutputGamma = IConsoleManager::Get().FindTConsoleVariableDataFloat(TEXT("r.TonemapperGamma"));
 
-		int32 OutputDeviceValue = CVarOutputDevice->GetValueOnRenderThread();
-		int32 OutputGamutValue = CVarOutputGamut->GetValueOnRenderThread();
+		int32 OutputDeviceValue = (int32)DisplayOutputFormat;
+		int32 OutputGamutValue = (int32)DisplayColorGamut;
 		float Gamma = CVarOutputGamma->GetValueOnRenderThread();
 
 		if (PLATFORM_APPLE && Gamma == 0.0f)
@@ -791,7 +783,7 @@ void RenderSlateBatch(FTexture2DRHIRef SlateRenderTarget, bool bClear, bool bIsH
 	}
 }
 
-inline bool CompositeUIWithHdrRenderTarget()
+inline bool CompositeUIWithHdrRenderTarget(const FViewportInfo* ViewInfo)
 {
 	// Optional off-screen UI composition during HDR rendering
 	static const auto CVarCompositeMode = IConsoleManager::Get().FindTConsoleVariableDataInt(TEXT("r.HDR.UI.CompositeMode"));
@@ -799,7 +791,7 @@ inline bool CompositeUIWithHdrRenderTarget()
 	const bool bSupportsUIComposition = GRHISupportsHDROutput && GSupportsVolumeTextureRendering && SupportsUICompositionRendering(GetFeatureLevelShaderPlatform(GMaxRHIFeatureLevel));
 	const bool bCompositeUI = bSupportsUIComposition
 		&& CVarCompositeMode && CVarCompositeMode->GetValueOnAnyThread() != 0
-		&& IsHDREnabled();
+		&& ViewInfo->bSceneHDREnabled;
 
 	return bCompositeUI;
 }
@@ -834,16 +826,13 @@ void FSlateRHIRenderer::DrawWindow_RenderThread(FRHICommandListImmediate& RHICmd
 		FMaterialRenderProxy::UpdateDeferredCachedUniformExpressions();
 		GetRendererModule().InitializeSystemTextures(RHICmdList);
 
-		const bool bCompositeUI = CompositeUIWithHdrRenderTarget();
+		const bool bCompositeUI = CompositeUIWithHdrRenderTarget(&ViewportInfo);
 
 		const int32 CompositionLUTSize = 32;
 
 		// Only need to update LUT on settings change
-		static const auto CVarHDROutputDevice = IConsoleManager::Get().FindTConsoleVariableDataInt(TEXT("r.HDR.Display.OutputDevice"));
-		static const auto CVarHDROutputGamut = IConsoleManager::Get().FindTConsoleVariableDataInt(TEXT("r.HDR.Display.ColorGamut"));
-
-		const int32 HDROutputDevice = CVarHDROutputDevice ? CVarHDROutputDevice->GetValueOnRenderThread() : (int32)EDisplayOutputFormat::SDR_sRGB;
-		const int32 HDROutputGamut = CVarHDROutputGamut ? CVarHDROutputGamut->GetValueOnRenderThread() : 0;
+		const int32 HDROutputDevice = (int32)ViewportInfo.HDRDisplayOutputFormat;
+		const int32 HDROutputGamut = (int32)ViewportInfo.HDRDisplayColorGamut;
 
 		bool bLUTStale = ViewportInfo.ColorSpaceLUTOutputDevice != HDROutputDevice || ViewportInfo.ColorSpaceLUTOutputGamut != HDROutputGamut;
 
@@ -929,8 +918,22 @@ void FSlateRHIRenderer::DrawWindow_RenderThread(FRHICommandListImmediate& RHICmd
 			{
 				bClear = true; // Force a clear of the UI buffer to black
 
-				// Grab HDR backbuffer
-				TransitionAndCopyTexture(RHICmdList, FinalBuffer, ViewportInfo.HDRSourceRT->GetRHI(), {});
+#if WITH_EDITOR
+				// in editor mode, we actually don't render to the true backbuffer, but to BufferedRT. The Scene image is rendered in the backbuffer with Slate
+				// We add it with FSlateElementBatcher::AddViewportElement and render it with RenderSlateBatch
+				if (WindowElementList.GetBatchDataHDR().GetRenderBatches().Num() > 0)
+				{
+					FRHIRenderPassInfo RPInfo(FinalBuffer, ERenderTargetActions::Clear_Store);
+					TransitionRenderPassTargets(RHICmdList, RPInfo);
+					RHICmdList.BeginRenderPass(RPInfo, TEXT("Clear back buffer"));
+					RHICmdList.EndRenderPass();
+				}
+				else
+#endif
+				{
+    				// Grab HDR backbuffer
+					TransitionAndCopyTexture(RHICmdList, FinalBuffer, ViewportInfo.HDRSourceRT->GetRHI(), {});
+				}
 
 				// UI backbuffer is temp target
 				BackBuffer = ViewportInfo.UITargetRT->GetRHI();
@@ -965,7 +968,7 @@ void FSlateRHIRenderer::DrawWindow_RenderThread(FRHICommandListImmediate& RHICmd
 			RHICmdList.BeginDrawingViewport(ViewportInfo.ViewportRHI, FTextureRHIRef());
 			RHICmdList.SetViewport(0, 0, 0, ViewportWidth, ViewportHeight, 0.0f);
 
-			bool bHdrTarget = IsHDREnabled() && !bCompositeUI;
+			bool bHdrTarget = ViewportInfo.bSceneHDREnabled && !bCompositeUI;
 			RenderSlateBatch(BackBuffer, bClear, bHdrTarget, ViewportInfo, ViewMatrix, BatchData, RHICmdList, ViewportWidth, ViewportHeight, DrawCommandParams, RenderingPolicy, PostProcessBuffer);
 
 			if (bCompositeUI)
@@ -1020,7 +1023,7 @@ void FSlateRHIRenderer::DrawWindow_RenderThread(FRHICommandListImmediate& RHICmd
 							GeometryShader->SetParameters(RHICmdList, VolumeBounds.MinZ);
 						}
 #endif
-						PixelShader->SetParameters(RHICmdList);
+						PixelShader->SetParameters(RHICmdList, ViewportInfo.HDRDisplayOutputFormat, ViewportInfo.HDRDisplayColorGamut);
 
 						RasterizeToVolumeTexture(RHICmdList, VolumeBounds);
 					}
@@ -1329,7 +1332,7 @@ void FSlateRHIRenderer::DrawWindows_Private(FSlateDrawBuffer& WindowDrawBuffer)
 				}
 
 				bool bCompositeHDRViewports = ElementBatcher->CompositeHDRViewports();
-				ElementBatcher->SetCompositeHDRViewports(CompositeUIWithHdrRenderTarget());
+				ElementBatcher->SetCompositeHDRViewports(CompositeUIWithHdrRenderTarget(ViewInfo));
 				// Add all elements for this window to the element batcher
 				ElementBatcher->AddElements(ElementList);
 				ElementBatcher->SetCompositeHDRViewports(bCompositeHDRViewports);
@@ -1363,7 +1366,7 @@ void FSlateRHIRenderer::DrawWindows_Private(FSlateDrawBuffer& WindowDrawBuffer)
 				if (Window->IsViewportSizeDrivenByWindow())
 				{
 					// Resize the viewport if needed
-					ConditionalResizeViewport(ViewInfo, ViewInfo->DesiredWidth, ViewInfo->DesiredHeight, IsViewportFullscreen(*Window));
+					ConditionalResizeViewport(ViewInfo, ViewInfo->DesiredWidth, ViewInfo->DesiredHeight, IsViewportFullscreen(*Window), Window);
 				}
 
 				// Tell the rendering thread to draw the windows
diff --git a/Engine/Source/Runtime/SlateRHIRenderer/Private/SlateRHIRenderer.h b/Engine/Source/Runtime/SlateRHIRenderer/Private/SlateRHIRenderer.h
--- a/Engine/Source/Runtime/SlateRHIRenderer/Private/SlateRHIRenderer.h
+++ b/Engine/Source/Runtime/SlateRHIRenderer/Private/SlateRHIRenderer.h
@@ -91,9 +91,9 @@ struct FViewportInfo : public FRenderResource
 	/** The desired SDR pixel format for this viewport */
 	EPixelFormat SDRPixelFormat;
 	/** Color gamut for output to HDR display */
-	int32 HDRColorGamut;
+	EDisplayColorGamut HDRDisplayColorGamut;
 	/** Device format for output to HDR display */
-	int32 HDROutputDevice;
+	EDisplayOutputFormat HDRDisplayOutputFormat;
 
 	IViewportRenderTargetProvider* RTProvider;
 	
@@ -120,6 +120,8 @@ struct FViewportInfo : public FRenderResource
 			bFullscreen(false),
 			PixelFormat(EPixelFormat::PF_Unknown),
 			SDRPixelFormat(EPixelFormat::PF_Unknown),
+			HDRDisplayColorGamut(EDisplayColorGamut::sRGB_D65),
+			HDRDisplayOutputFormat(EDisplayOutputFormat::SDR_sRGB),
 			RTProvider(nullptr),
 			bHDREnabled(false),
 			bSceneHDREnabled(false)
@@ -253,7 +255,7 @@ private:
 	 * @param Height		The height that we shoudl size to
 	 * @param bFullscreen	If we should be in fullscreen
 	 */
-	void ConditionalResizeViewport( FViewportInfo* ViewportInfo, uint32 Width, uint32 Height, bool bFullscreen );
+	void ConditionalResizeViewport( FViewportInfo* ViewportInfo, uint32 Width, uint32 Height, bool bFullscreen, SWindow* Window);
 	
 	/** 
 	 * Creates necessary resources to render a window and sends draw commands to the rendering thread
diff --git a/Engine/Source/Runtime/Windows/D3D11RHI/Private/D3D11RHIPrivate.h b/Engine/Source/Runtime/Windows/D3D11RHI/Private/D3D11RHIPrivate.h
--- a/Engine/Source/Runtime/Windows/D3D11RHI/Private/D3D11RHIPrivate.h
+++ b/Engine/Source/Runtime/Windows/D3D11RHI/Private/D3D11RHIPrivate.h
@@ -16,6 +16,7 @@
 #include "Containers/ResourceArray.h"
 #include "EngineGlobals.h"
 #include "Engine/Engine.h"
+#include "HDRHelper.h"
 
 DECLARE_LOG_CATEGORY_EXTERN(LogD3D11RHI, Log, All);
 
@@ -677,6 +678,7 @@ public:
 	virtual ID3D11RenderTargetView* RHIGetRenderTargetView(FRHITexture* InTexture, int32 InMipIndex = 0, int32 InArraySliceIndex = -1) const final override;
 	virtual ID3D11ShaderResourceView* RHIGetShaderResourceView(FRHITexture* InTexture) const final override;
 	virtual void                  RHIRegisterWork(uint32 NumPrimitives) final override;
+	virtual void				  RHIGetDisplaysInformation(FDisplayInformationArray& OutDisplayInformation) final override;
 
 	// Accessors.
 	ID3D11Device* GetDevice() const
@@ -922,6 +924,8 @@ protected:
 	uint32 HDRDetectedDisplayIndex;
 	uint32 HDRDetectedDisplayIHVIndex;
 
+	FDisplayInformationArray DisplayList;
+
 	bool bRenderDoc = false;
 
 public:
@@ -1050,6 +1054,8 @@ protected:
 	double IntelMetricsDicoveryGetGPUTime();
 #endif
 
+	bool SetupDisplayHDRMetaData();
+
 	friend struct FD3DGPUProfiler;
 
 };
diff --git a/Engine/Source/Runtime/Windows/D3D11RHI/Private/D3D11Viewport.cpp b/Engine/Source/Runtime/Windows/D3D11RHI/Private/D3D11Viewport.cpp
--- a/Engine/Source/Runtime/Windows/D3D11RHI/Private/D3D11Viewport.cpp
+++ b/Engine/Source/Runtime/Windows/D3D11RHI/Private/D3D11Viewport.cpp
@@ -6,6 +6,7 @@
 
 #include "D3D11RHIPrivate.h"
 #include "RenderCore.h"
+#include "HDRHelper.h"
 #include "Engine/RendererSettings.h"
 #include "HAL/ThreadHeartBeat.h"
 
@@ -312,6 +313,15 @@ void FD3D11Viewport::Resize(uint32 InSizeX, uint32 InSizeY, bool bInIsFullscreen
 		}
 	}
 
+	RECT WindowRect = {};
+#if PLATFORM_WINDOWS
+	GetWindowRect(WindowHandle, &WindowRect);
+#endif
+	FVector2D WindowTopLeft((float)WindowRect.left, (float)WindowRect.top);
+	FVector2D WindowBottomRight((float)WindowRect.right, (float)WindowRect.bottom);
+	bool bHDREnabled;
+	HDRGetMetaData(DisplayOutputFormat, DisplayColorGamut, bHDREnabled, WindowTopLeft, WindowBottomRight, (void*)WindowHandle);
+
 	// Float RGBA backbuffers are requested whenever HDR mode is desired
 	if (PixelFormat == GRHIHDRDisplayOutputFormat && bIsFullscreen)
 	{
diff --git a/Engine/Source/Runtime/Windows/D3D11RHI/Private/Windows/WindowsD3D11Device.cpp b/Engine/Source/Runtime/Windows/D3D11RHI/Private/Windows/WindowsD3D11Device.cpp
--- a/Engine/Source/Runtime/Windows/D3D11RHI/Private/Windows/WindowsD3D11Device.cpp
+++ b/Engine/Source/Runtime/Windows/D3D11RHI/Private/Windows/WindowsD3D11Device.cpp
@@ -472,12 +472,10 @@ static void SetHDRMonitorModeAMD(uint32 IHVDisplayIndex, bool bEnableHDR, EDispl
 /** Enable HDR meta data transmission */
 void FD3D11DynamicRHI::EnableHDR()
 {
-	static const auto CVarHDRColorGamut = IConsoleManager::Get().FindTConsoleVariableDataInt(TEXT("r.HDR.Display.ColorGamut"));
-	static const auto CVarHDROutputDevice = IConsoleManager::Get().FindTConsoleVariableDataInt(TEXT("r.HDR.Display.OutputDevice"));
-
 	if ( GRHISupportsHDROutput && IsHDREnabled() )
 	{
-		const EDisplayOutputFormat OutputDevice = (EDisplayOutputFormat)CVarHDROutputDevice->GetValueOnAnyThread();
+		const EDisplayOutputFormat OutputDevice = HDRGetDefaultDisplayOutputFormat();
+		const EDisplayColorGamut DisplayGamut = HDRGetDefaultDisplayColorGamut();
 
 		const float DisplayMaxOutputNits = (OutputDevice == EDisplayOutputFormat::HDR_ACES_2000nit_ST2084 || OutputDevice == EDisplayOutputFormat::HDR_ACES_2000nit_ScRGB) ? 2000.f : 1000.f;
 		const float DisplayMinOutputNits = 0.0f;	// Min output of the display
@@ -489,7 +487,7 @@ void FD3D11DynamicRHI::EnableHDR()
 			SetHDRMonitorModeNVIDIA(
 				HDRDetectedDisplayIHVIndex,
 				true,
-				EDisplayColorGamut(CVarHDRColorGamut->GetValueOnAnyThread()),
+				DisplayGamut,
 				DisplayMaxOutputNits,
 				DisplayMinOutputNits,
 				DisplayMaxCLL,
@@ -500,7 +498,7 @@ void FD3D11DynamicRHI::EnableHDR()
 			SetHDRMonitorModeAMD(
 				HDRDetectedDisplayIHVIndex,
 				true,
-				EDisplayColorGamut(CVarHDRColorGamut->GetValueOnAnyThread()),
+				DisplayGamut,
 				DisplayMaxOutputNits,
 				DisplayMinOutputNits,
 				DisplayMaxCLL,
@@ -553,19 +551,19 @@ void FD3D11DynamicRHI::ShutdownHDR()
 	}
 }
 
-static bool SupportsHDROutput(FD3D11DynamicRHI* D3DRHI)
+bool FD3D11DynamicRHI::SetupDisplayHDRMetaData()
 {
-	check(D3DRHI && D3DRHI->GetDevice());
-	ID3D11Device* Direct3DDevice = D3DRHI->GetDevice();
+	check(GetDevice());
 
 	// Default to primary display
-	D3DRHI->SetHDRDetectedDisplayIndices(0, 0);
+	SetHDRDetectedDisplayIndices(0, 0);
 	
+	DisplayList.Empty();
+
 #if WITH_EDITOR
 	// Determines if any displays support HDR
 	bool bSupportsHDROutput = false;
 	{
-		const FD3D11Adapter& Adapter = D3DRHI->GetAdapter();
 		IDXGIAdapter* DXGIAdapter = Adapter.DXGIAdapter;
 		if (!DXGIAdapter)
 		{
@@ -597,6 +595,12 @@ static bool SupportsHDROutput(FD3D11DynamicRHI* D3DRHI)
 
 					bSupportsHDROutput = true;
 				}
+
+				FDisplayInformation DisplayInformation{};
+				DisplayInformation.bHDRSupported = bDisplaySupportsHDROutput;
+				const RECT& DisplayCoords = OutputDesc.DesktopCoordinates;
+				DisplayInformation.DesktopCoordinates = FIntRect(DisplayCoords.left, DisplayCoords.top, DisplayCoords.right, DisplayCoords.bottom);
+				DisplayList.Add(DisplayInformation);
 			}
 		}
 	}
@@ -615,6 +619,7 @@ static bool SupportsHDROutput(FD3D11DynamicRHI* D3DRHI)
 	uint32 ForcedDisplayIndex = 0;
 	bool bForcedDisplay = FParse::Value(FCommandLine::Get(), TEXT("FullscreenDisplay="), ForcedDisplayIndex);
 
+	bool bSupportsHDROutput = false;
 	for (; true; ++DisplayIndex)
 	{
 		TRefCountPtr<IDXGIOutput> DXGIOutput;
@@ -631,7 +636,11 @@ static bool SupportsHDROutput(FD3D11DynamicRHI* D3DRHI)
 
 		DXGI_OUTPUT_DESC OutputDesc;
 		DXGIOutput->GetDesc(&OutputDesc);
-		
+		FDisplayInformation DisplayInformation{};
+		const RECT& DisplayCoords = OutputDesc.DesktopCoordinates;
+		DisplayInformation.DesktopCoordinates = FIntRect(DisplayCoords.left, DisplayCoords.top, DisplayCoords.right, DisplayCoords.bottom);
+		DisplayInformation.bHDRSupported = false;
+
 		if (IsRHIDeviceNVIDIA())
 		{
 #ifdef NVAPI_INTERFACE
@@ -648,12 +657,14 @@ static bool SupportsHDROutput(FD3D11DynamicRHI* D3DRHI)
 
 				if (NVAPI_OK == NvAPI_Disp_GetHdrCapabilities(DisplayId, &HdrCapabilities))
 				{		
-					if (HdrCapabilities.isST2084EotfSupported)
+					// we're only choosing the first supported HDR monitor
+					if (HdrCapabilities.isST2084EotfSupported && !bSupportsHDROutput)
 					{
 						UE_LOG(LogD3D11RHI, Log, TEXT("HDR output is supported on display %i (NvId: 0x%x)."), DisplayIndex, DisplayId);
-						D3DRHI->SetHDRDetectedDisplayIndices(DisplayIndex, DisplayId);
-						return true;
+						SetHDRDetectedDisplayIndices(DisplayIndex, DisplayId);
+						bSupportsHDROutput = true;
 					}
+					DisplayInformation.bHDRSupported = HdrCapabilities.isST2084EotfSupported;
 				}
 			}
 			else if (Status != NVAPI_ERROR && Status != NVAPI_NVIDIA_DEVICE_NOT_FOUND)
@@ -678,12 +689,14 @@ static bool SupportsHDROutput(FD3D11DynamicRHI* D3DRHI)
 					{
 						// AGS has flags for HDR10 and Dolby Vision instead of a flag for the ST2084 transfer function.
 						// Both HDR10 and Dolby Vision use the ST2084 EOTF.
-						if (DisplayInfo.HDR10 != 0 || DisplayInfo.dolbyVision != 0)
+						bool DisplaySupportsHDR = (DisplayInfo.HDR10 != 0 || DisplayInfo.dolbyVision != 0);
+						if (DisplaySupportsHDR && !bSupportsHDROutput)
 						{
 							UE_LOG(LogD3D11RHI, Log, TEXT("HDR output is supported on display %i (AMD Device: 0x%x, Display: 0x%x)."), DisplayIndex, AMDDeviceIndex, AMDDisplayIndex);
-							D3DRHI->SetHDRDetectedDisplayIndices(DisplayIndex, (uint32)(AMDDeviceIndex << 16) | (uint32)AMDDisplayIndex);
-							return true;
+							SetHDRDetectedDisplayIndices(DisplayIndex, (uint32)(AMDDeviceIndex << 16) | (uint32)AMDDisplayIndex);
+							bSupportsHDROutput = true;
 						}
+						DisplayInformation.bHDRSupported = DisplaySupportsHDR;
 					}
 				}
 			}
@@ -693,9 +706,10 @@ static bool SupportsHDROutput(FD3D11DynamicRHI* D3DRHI)
 		{
 			// Not yet implemented
 		}
+		DisplayList.Add(DisplayInformation);
 	}
 
-	return false;
+	return bSupportsHDROutput;
 #endif
 }
 
@@ -2204,7 +2218,7 @@ void FD3D11DynamicRHI::InitD3DDevice()
 #endif
 		
 		{
-			GRHISupportsHDROutput = SupportsHDROutput(this);
+			GRHISupportsHDROutput = SetupDisplayHDRMetaData();
 		}
 
 		// Add device overclock state to crash context
diff --git a/Engine/Source/Runtime/Windows/D3D11RHI/Private/Windows/WindowsD3D11Viewport.cpp b/Engine/Source/Runtime/Windows/D3D11RHI/Private/Windows/WindowsD3D11Viewport.cpp
--- a/Engine/Source/Runtime/Windows/D3D11RHI/Private/Windows/WindowsD3D11Viewport.cpp
+++ b/Engine/Source/Runtime/Windows/D3D11RHI/Private/Windows/WindowsD3D11Viewport.cpp
@@ -44,6 +44,8 @@ FD3D11Viewport::FD3D11Viewport(FD3D11DynamicRHI* InD3DRHI,HWND InWindowHandle,ui
 	ValidState(0),
 	PixelFormat(InPreferredPixelFormat),
 	PixelColorSpace(EColorSpaceAndEOTF::ERec709_sRGB),
+	DisplayColorGamut(EDisplayColorGamut::sRGB_D65),
+	DisplayOutputFormat(EDisplayOutputFormat::SDR_sRGB),
 	bIsFullscreen(bInIsFullscreen),
 	bAllowTearing(false),
 	FrameSyncEvent(InD3DRHI)
@@ -343,12 +345,7 @@ void FD3D11Viewport::CheckHDRMonitorStatus()
 	{
 		FlushRenderingCommands();
 
-		static const auto CVarHDRColorGamut = IConsoleManager::Get().FindTConsoleVariableDataInt(TEXT("r.HDR.Display.ColorGamut"));
-		const EDisplayColorGamut DisplayGamut = EDisplayColorGamut(CVarHDRColorGamut->GetValueOnAnyThread());
-		static const auto CVarHDROutputDevice = IConsoleManager::Get().FindTConsoleVariableDataInt(TEXT("r.HDR.Display.OutputDevice"));
-		const EDisplayOutputFormat OutputDevice = EDisplayOutputFormat(CVarHDROutputDevice->GetValueOnAnyThread());
-
-		EnsureColorSpace(SwapChain, DisplayGamut, OutputDevice, PixelFormat);
+		EnsureColorSpace(SwapChain, DisplayColorGamut, DisplayOutputFormat, PixelFormat);
 	}
 	
 	{
@@ -405,4 +402,9 @@ void FD3D11Viewport::ResetSwapChainInternal(bool bIgnoreFocus)
 	}
 }
 
+void FD3D11DynamicRHI::RHIGetDisplaysInformation(FDisplayInformationArray& OutDisplayInformation)
+{
+	OutDisplayInformation.Append(DisplayList);
+}
+
 #include "Windows/HideWindowsPlatformTypes.h"
\ No newline at end of file
diff --git a/Engine/Source/Runtime/Windows/D3D11RHI/Public/D3D11Viewport.h b/Engine/Source/Runtime/Windows/D3D11RHI/Public/D3D11Viewport.h
--- a/Engine/Source/Runtime/Windows/D3D11RHI/Public/D3D11Viewport.h
+++ b/Engine/Source/Runtime/Windows/D3D11RHI/Public/D3D11Viewport.h
@@ -143,6 +143,8 @@ protected:
 	TAtomic<uint32> ValidState;
 	EPixelFormat PixelFormat;
 	EColorSpaceAndEOTF PixelColorSpace;
+	EDisplayColorGamut DisplayColorGamut;
+	EDisplayOutputFormat DisplayOutputFormat;
 	bool bIsFullscreen;
 	bool bAllowTearing;
 
