RHI: Fix crash in Vulkan on Linux (and leak on Windows) with modal slate dialogs.

diff --git a/Engine/Source/Runtime/Slate/Private/Framework/Application/SlateApplication.cpp b/Engine/Source/Runtime/Slate/Private/Framework/Application/SlateApplication.cpp
--- a/Engine/Source/Runtime/Slate/Private/Framework/Application/SlateApplication.cpp
+++ b/Engine/Source/Runtime/Slate/Private/Framework/Application/SlateApplication.cpp
@@ -2069,6 +2069,9 @@ void FSlateApplication::AddModalWindow( TSharedRef<SWindow> InSlateWindow, const
 				Renderer->EndFrame();
 			}
 
+			// Make sure every fake frame of our modal dialog has a chance to free its resources
+			Renderer->FlushCommands();
+
 			// Synchronize the game thread and the render thread so that the render thread doesn't get too far behind.
 			Renderer->Sync();
 		}
diff --git a/Engine/Source/Runtime/SlateRHIRenderer/Private/SlateRHIRenderer.cpp b/Engine/Source/Runtime/SlateRHIRenderer/Private/SlateRHIRenderer.cpp
--- a/Engine/Source/Runtime/SlateRHIRenderer/Private/SlateRHIRenderer.cpp
+++ b/Engine/Source/Runtime/SlateRHIRenderer/Private/SlateRHIRenderer.cpp
@@ -1551,7 +1551,9 @@ void FSlateRHIRenderer::BeginFrame() const
 	ENQUEUE_RENDER_COMMAND(SlateRHIBeginFrame)(
 	   [](FRHICommandListImmediate& RHICmdList)
 	   {
-		   RHICmdList.BeginFrame();
+			GFrameNumberRenderThread++;
+			GPU_STATS_BEGINFRAME(RHICmdList);
+			RHICmdList.BeginFrame();
 	   }
 	);
 }
@@ -1562,7 +1564,8 @@ void FSlateRHIRenderer::EndFrame() const
 	   [](FRHICommandListImmediate& RHICmdList)
 	   {
 		   RHICmdList.EndFrame();
-	   }
+		   GPU_STATS_ENDFRAME(RHICmdList);
+		}
 	);
 }
 
