Slate: Fix implicit double to float conversions in the SlateRHIRenderer module.

diff --git a/Engine/Source/Runtime/SlateCore/Public/Types/Attributes/SlateAttributeBase.inl b/Engine/Source/Runtime/SlateCore/Public/Types/Attributes/SlateAttributeBase.inl
--- a/Engine/Source/Runtime/SlateCore/Public/Types/Attributes/SlateAttributeBase.inl
+++ b/Engine/Source/Runtime/SlateCore/Public/Types/Attributes/SlateAttributeBase.inl
@@ -68,9 +68,13 @@ namespace SlateAttributePrivate
 			{
 				ensureMsgf(FMath::IsFinite(Value), TEXT("Value contains a NaN. Initialize your double properly"));
 			}
-			else if constexpr (std::is_same<FVector2D, InObjectType>::value)
+			else if constexpr (std::is_same<FVector2f, InObjectType>::value)
 			{
-				ensureMsgf(!Value.ContainsNaN(), TEXT("Value contains a NaN. Initialize your FVector2D properly (see FVector2D::EForceInit)"));
+				ensureMsgf(!Value.ContainsNaN(), TEXT("Value contains a NaN. Initialize your FVector2f properly (see FVector2D::EForceInit)"));
+			}
+			else if constexpr (std::is_same<FVector2d, InObjectType>::value)
+			{
+				ensureMsgf(!Value.ContainsNaN(), TEXT("Value contains a NaN. Initialize your FVector2d properly (see FVector2D::EForceInit)"));
 			}
 			else if constexpr (std::is_same<FVector3f, InObjectType>::value)
 			{
@@ -80,9 +84,13 @@ namespace SlateAttributePrivate
 			{
 				ensureMsgf(!Value.ContainsNaN(), TEXT("Value contains a NaN. Initialize your FVector3d properly (see FVector::EForceInit)"));
 			}
-			else if constexpr (std::is_same<FVector4, InObjectType>::value)
+			else if constexpr (std::is_same<FVector4f, InObjectType>::value)
+			{
+				ensureMsgf(!Value.ContainsNaN(), TEXT("Value contains a NaN. Initialize your FVector4f properly"));
+			}
+			else if constexpr (std::is_same<FVector4d, InObjectType>::value)
 			{
-				ensureMsgf(!Value.ContainsNaN(), TEXT("Value contains a NaN. Initialize your FVector4 properly"));
+				ensureMsgf(!Value.ContainsNaN(), TEXT("Value contains a NaN. Initialize your FVector4d properly"));
 			}
 			else if constexpr (std::is_same<FLinearColor, InObjectType>::value)
 			{
diff --git a/Engine/Source/Runtime/SlateRHIRenderer/Private/Slate3DRenderer.cpp b/Engine/Source/Runtime/SlateRHIRenderer/Private/Slate3DRenderer.cpp
--- a/Engine/Source/Runtime/SlateRHIRenderer/Private/Slate3DRenderer.cpp
+++ b/Engine/Source/Runtime/SlateRHIRenderer/Private/Slate3DRenderer.cpp
@@ -6,6 +6,7 @@
 #include "SceneUtils.h"
 #include "SlateRHIRenderer.h"
 #include "Rendering/ElementBatcher.h"
+#include "Types/SlateVector2.h"
 
 DECLARE_GPU_STAT_NAMED(Slate3D, TEXT("Slate 3D"));
 
@@ -117,7 +118,7 @@ void FSlate3DRenderer::DrawWindow_GameThread(FSlateDrawBuffer& DrawBuffer)
 		if (Window)
 		{
 			const FVector2D WindowSize = Window->GetSizeInScreen();
-			if (WindowSize.X > 0 && WindowSize.Y > 0)
+			if (WindowSize.X > 0.0 && WindowSize.Y > 0.0)
 			{
 				// Add all elements for this window to the element batcher
 				ElementBatcher->AddElements(ElementList);
@@ -177,7 +178,7 @@ void FSlate3DRenderer::DrawWindowToTarget_RenderThread(FRHICommandListImmediate&
 			FVector2D DrawOffset = Context.WindowDrawBuffer->ViewOffset;
 
 			FMatrix ProjectionMatrix = FSlateRHIRenderer::CreateProjectionMatrix(RTTextureRHI->GetSizeX(), RTTextureRHI->GetSizeY());
-			FMatrix ViewOffset = FTranslationMatrix::Make(FVector(DrawOffset, 0));
+			FMatrix ViewOffset = FTranslationMatrix::Make(FVector(DrawOffset, 0.0));
 			ProjectionMatrix = ViewOffset * ProjectionMatrix;
 
 			FSlateBackBuffer BackBufferTarget(Context.RenderTarget->GetRenderTargetTexture(), FIntPoint(RTTextureRHI->GetSizeX(), RTTextureRHI->GetSizeY()));
@@ -185,7 +186,7 @@ void FSlate3DRenderer::DrawWindowToTarget_RenderThread(FRHICommandListImmediate&
 			FSlateRenderingParams DrawOptions(ProjectionMatrix, FGameTime::CreateDilated(Context.RealTimeSeconds, Context.DeltaRealTimeSeconds, Context.WorldTimeSeconds, Context.DeltaTimeSeconds));
 			// The scene renderer will handle it in this case
 			DrawOptions.bAllowSwitchVerticalAxis = false;
-			DrawOptions.ViewOffset = DrawOffset;
+			DrawOptions.ViewOffset = UE::Slate::CastToVector2f(DrawOffset);
 
 			FTexture2DRHIRef ColorTarget = Context.RenderTarget->GetRenderTargetTexture();
 
diff --git a/Engine/Source/Runtime/SlateRHIRenderer/Private/SlateMaterialResource.cpp b/Engine/Source/Runtime/SlateRHIRenderer/Private/SlateMaterialResource.cpp
--- a/Engine/Source/Runtime/SlateRHIRenderer/Private/SlateMaterialResource.cpp
+++ b/Engine/Source/Runtime/SlateRHIRenderer/Private/SlateMaterialResource.cpp
@@ -62,7 +62,7 @@ namespace SlateMaterialResource
 
 }
 
-FSlateMaterialResource::FSlateMaterialResource(const UMaterialInterface& InMaterialResource, const FVector2D& InImageSize, FSlateShaderResource* InTextureMask )
+FSlateMaterialResource::FSlateMaterialResource(const UMaterialInterface& InMaterialResource, const FVector2f InImageSize, FSlateShaderResource* InTextureMask )
 	: MaterialObject( &InMaterialResource)
 	, SlateProxy( new FSlateShaderResourceProxy )
 	, TextureMaskResource( InTextureMask )
@@ -99,7 +99,7 @@ FSlateMaterialResource::~FSlateMaterialResource()
 	}
 }
 
-void FSlateMaterialResource::UpdateMaterial(const UMaterialInterface& InMaterialResource, const FVector2D& InImageSize, FSlateShaderResource* InTextureMask)
+void FSlateMaterialResource::UpdateMaterial(const UMaterialInterface& InMaterialResource, const FVector2f InImageSize, FSlateShaderResource* InTextureMask)
 {
 #if SLATE_CHECK_UOBJECT_RENDER_RESOURCES
 	SlateMaterialResource::CheckInvalidUMaterial(InMaterialResource, DebugName);
diff --git a/Engine/Source/Runtime/SlateRHIRenderer/Private/SlateMaterialResource.h b/Engine/Source/Runtime/SlateRHIRenderer/Private/SlateMaterialResource.h
--- a/Engine/Source/Runtime/SlateRHIRenderer/Private/SlateMaterialResource.h
+++ b/Engine/Source/Runtime/SlateRHIRenderer/Private/SlateMaterialResource.h
@@ -14,14 +14,14 @@ class FMaterialRenderProxy;
 class FSlateMaterialResource : public FSlateShaderResource
 {
 public:
-	FSlateMaterialResource(const UMaterialInterface& InMaterialResource, const FVector2D& InImageSize, FSlateShaderResource* InTextureMask = nullptr );
+	FSlateMaterialResource(const UMaterialInterface& InMaterialResource, const FVector2f InImageSize, FSlateShaderResource* InTextureMask = nullptr );
 	~FSlateMaterialResource();
 
 	virtual uint32 GetWidth() const override { return Width; }
 	virtual uint32 GetHeight() const override { return Height; }
 	virtual ESlateShaderResource::Type GetType() const override { return ESlateShaderResource::Material; }
 
-	void UpdateMaterial(const UMaterialInterface& InMaterialResource, const FVector2D& InImageSize, FSlateShaderResource* InTextureMask );
+	void UpdateMaterial(const UMaterialInterface& InMaterialResource, const FVector2f InImageSize, FSlateShaderResource* InTextureMask );
 	void ResetMaterial();
 
 	/** @return The material render proxy */
diff --git a/Engine/Source/Runtime/SlateRHIRenderer/Private/SlateMaterialShader.cpp b/Engine/Source/Runtime/SlateRHIRenderer/Private/SlateMaterialShader.cpp
--- a/Engine/Source/Runtime/SlateRHIRenderer/Private/SlateMaterialShader.cpp
+++ b/Engine/Source/Runtime/SlateRHIRenderer/Private/SlateMaterialShader.cpp
@@ -136,7 +136,7 @@ void FSlateMaterialShaderPS::SetDisplayGammaAndContrast(FRHICommandList& RHICmdL
 
 void FSlateMaterialShaderPS::SetDrawFlags(FRHICommandList& RHICmdList, bool bDrawDisabledEffect)
 {
-	FVector4f InDrawFlags(bDrawDisabledEffect ? 1 : 0, 0, 0, 0);
+	FVector4f InDrawFlags((bDrawDisabledEffect ? 1.f : 0.f), 0.f, 0.f, 0.f);
 
 	SetShaderValue(RHICmdList, RHICmdList.GetBoundPixelShader(), DrawFlags, InDrawFlags);
 }
diff --git a/Engine/Source/Runtime/SlateRHIRenderer/Private/SlatePostProcessor.cpp b/Engine/Source/Runtime/SlateRHIRenderer/Private/SlatePostProcessor.cpp
--- a/Engine/Source/Runtime/SlateRHIRenderer/Private/SlatePostProcessor.cpp
+++ b/Engine/Source/Runtime/SlateRHIRenderer/Private/SlatePostProcessor.cpp
@@ -85,10 +85,10 @@ void FSlatePostProcessor::BlurRect(FRHICommandListImmediate& RHICmdList, IRender
 	GraphicsPSOInit.DepthStencilState = TStaticDepthStencilState<false, CF_Always>::GetRHI();
 
 	RHICmdList.SetScissorRect(false, 0, 0, 0, 0);
-	RHICmdList.SetViewport(0, 0, 0, DestTextureWidth, DestTextureHeight, 0.0f);
+	RHICmdList.SetViewport(0.f, 0.f, 0.f, (float)DestTextureWidth, (float)DestTextureHeight, 0.0f);
 	
-	const FVector2D InvBufferSize = FVector2D(1.0f / DestTextureWidth, 1.0f / DestTextureHeight);
-	const FVector2D HalfTexelOffset = FVector2D(0.5f/ DestTextureWidth, 0.5f/ DestTextureHeight);
+	const FVector2f InvBufferSize = FVector2f(1.0f / (float)DestTextureWidth, 1.0f / (float)DestTextureHeight);
+	const FVector2f HalfTexelOffset = FVector2f(0.5f/ (float)DestTextureWidth, 0.5f/ (float)DestTextureHeight);
 
 	for (int32 PassIndex = 0; PassIndex < 2; ++PassIndex)
 	{
@@ -117,15 +117,15 @@ void FSlatePostProcessor::BlurRect(FRHICommandListImmediate& RHICmdList, IRender
 
 				if (bDownsample)
 				{
-					PixelShader->SetUVBounds(RHICmdList, FVector4f(FVector2f::ZeroVector, FVector2f((float)DownsampleSize.X / DestTextureWidth, (float)DownsampleSize.Y / DestTextureHeight) - FVector2f(HalfTexelOffset)));
-					PixelShader->SetBufferSizeAndDirection(RHICmdList, InvBufferSize, FVector2D(1, 0));
+					PixelShader->SetUVBounds(RHICmdList, FVector4f(FVector2f::ZeroVector, FVector2f((float)DownsampleSize.X / DestTextureWidth, (float)DownsampleSize.Y / DestTextureHeight) - HalfTexelOffset));
+					PixelShader->SetBufferSizeAndDirection(RHICmdList, InvBufferSize, FVector2f(1.f, 0.f));
 
 					RendererModule.DrawRectangle(
 						RHICmdList,
+						0.f, 0.f,
+						(float)DownsampleSize.X, (float)DownsampleSize.Y,
 						0, 0,
-						DownsampleSize.X, DownsampleSize.Y,
-						0, 0,
-						DownsampleSize.X, DownsampleSize.Y,
+						(float)DownsampleSize.X, (float)DownsampleSize.Y,
 						FIntPoint(DestTextureWidth, DestTextureHeight),
 						FIntPoint(DestTextureWidth, DestTextureHeight),
 						VertexShader,
@@ -140,12 +140,12 @@ void FSlatePostProcessor::BlurRect(FRHICommandListImmediate& RHICmdList, IRender
 					const FVector2f SizeUV = UVEnd - UVStart;
 
 					PixelShader->SetUVBounds(RHICmdList, FVector4f(UVStart, UVEnd));
-					PixelShader->SetBufferSizeAndDirection(RHICmdList, FVector2D(InvSrcTextureSize), FVector2D(1, 0));
+					PixelShader->SetBufferSizeAndDirection(RHICmdList, InvSrcTextureSize, FVector2f(1.f, 0.f));
 
 					RendererModule.DrawRectangle(
 						RHICmdList,
-						0, 0,
-						RequiredSize.X, RequiredSize.Y,
+						0.f, 0.f,
+						(float)RequiredSize.X, (float)RequiredSize.Y,
 						UVStart.X, UVStart.Y,
 						SizeUV.X, SizeUV.Y,
 						FIntPoint(DestTextureWidth, DestTextureHeight),
@@ -176,16 +176,16 @@ void FSlatePostProcessor::BlurRect(FRHICommandListImmediate& RHICmdList, IRender
 				SetGraphicsPipelineState(RHICmdList, GraphicsPSOInit, 0);
 
 				PixelShader->SetWeightsAndOffsets(RHICmdList, WeightsAndOffsets, SampleCount);
-				PixelShader->SetUVBounds(RHICmdList, FVector4f(FVector2f::ZeroVector, FVector2f((float)DownsampleSize.X / DestTextureWidth, (float)DownsampleSize.Y / DestTextureHeight) - FVector2f(HalfTexelOffset)));
+				PixelShader->SetUVBounds(RHICmdList, FVector4f(FVector2f::ZeroVector, FVector2f((float)DownsampleSize.X / DestTextureWidth, (float)DownsampleSize.Y / DestTextureHeight) - HalfTexelOffset));
 				PixelShader->SetTexture(RHICmdList, SourceTexture, BilinearClamp);
-				PixelShader->SetBufferSizeAndDirection(RHICmdList, InvBufferSize, FVector2D(0, 1));
+				PixelShader->SetBufferSizeAndDirection(RHICmdList, InvBufferSize, FVector2f(0.f, 1.f));
 
 				RendererModule.DrawRectangle(
 					RHICmdList,
-					0, 0,
-					DownsampleSize.X, DownsampleSize.Y,
-					0, 0,
-					DownsampleSize.X, DownsampleSize.Y,
+					0.f, 0.f,
+					(float)DownsampleSize.X, (float)DownsampleSize.Y,
+					0.f, 0.f,
+					(float)DownsampleSize.X, (float)DownsampleSize.Y,
 					FIntPoint(DestTextureWidth, DestTextureHeight),
 					FIntPoint(DestTextureWidth, DestTextureHeight),
 					VertexShader,
@@ -236,7 +236,7 @@ void FSlatePostProcessor::ColorDeficiency(FRHICommandListImmediate& RHICmdList,
 	GraphicsPSOInit.DepthStencilState = TStaticDepthStencilState<false, CF_Always>::GetRHI();
 
 	RHICmdList.SetScissorRect(false, 0, 0, 0, 0);
-	RHICmdList.SetViewport(0, 0, 0, DestTextureWidth, DestTextureHeight, 0.0f);
+	RHICmdList.SetViewport(0.f, 0.f, 0.f, (float)DestTextureWidth, (float)DestTextureHeight, 0.0f);
 
 	// 
 	{
@@ -263,10 +263,10 @@ void FSlatePostProcessor::ColorDeficiency(FRHICommandListImmediate& RHICmdList,
 
 			RendererModule.DrawRectangle(
 				RHICmdList,
-				0, 0,
-				RequiredSize.X, RequiredSize.Y,
-				0, 0,
-				1, 1,
+				0.f, 0,
+				(float)RequiredSize.X, (float)RequiredSize.Y,
+				0.f, 0.f,
+				1.f, 1.f,
 				FIntPoint(DestTextureWidth, DestTextureHeight),
 				FIntPoint(1, 1),
 				VertexShader,
@@ -326,7 +326,7 @@ void FSlatePostProcessor::DownsampleRect(FRHICommandListImmediate& RHICmdList, I
 		const FVector2f UVEnd = FVector2f(DestRect.Right, DestRect.Bottom) * InvSrcTextureSize;
 		const FVector2f SizeUV = UVEnd - UVStart;
 		
-		RHICmdList.SetViewport(0, 0, 0, DestTextureWidth, DestTextureHeight, 0.0f);
+		RHICmdList.SetViewport(0.f, 0.f, 0.f, (float)DestTextureWidth, (float)DestTextureHeight, 0.0f);
 		RHICmdList.SetScissorRect(false, 0, 0, 0, 0);
 
 		FRHIRenderPassInfo RPInfo(DestTexture, ERenderTargetActions::Load_Store);
@@ -350,8 +350,8 @@ void FSlatePostProcessor::DownsampleRect(FRHICommandListImmediate& RHICmdList, I
 
 			RendererModule.DrawRectangle(
 				RHICmdList,
-				0, 0,
-				DownsampleSize.X, DownsampleSize.Y,
+				0.f, 0.f,
+				(float)DownsampleSize.X, (float)DownsampleSize.Y,
 				UVStart.X, UVStart.Y,
 				SizeUV.X, SizeUV.Y,
 				FIntPoint(DestTextureWidth, DestTextureHeight),
@@ -372,7 +372,7 @@ void FSlatePostProcessor::UpsampleRect(FRHICommandListImmediate& RHICmdList, IRe
 {
 	SCOPED_DRAW_EVENT(RHICmdList, SlatePostProcessUpsample);
 
-	const FVector4 Zero(0, 0, 0, 0);
+	const FVector4f Zero(0.f, 0.f, 0.f, 0.f);
 
 	FGraphicsPipelineStateInitializer GraphicsPSOInit;
 	GraphicsPSOInit.BlendState = Params.CornerRadius == Zero ? TStaticBlendState<>::GetRHI() : TStaticBlendState<CW_RGBA, BO_Add, BF_SourceAlpha, BF_InverseSourceAlpha, BO_Add, BF_One, BF_InverseSourceAlpha>::GetRHI();
@@ -398,7 +398,7 @@ void FSlatePostProcessor::UpsampleRect(FRHICommandListImmediate& RHICmdList, IRe
 	FGlobalShaderMap* ShaderMap = GetGlobalShaderMap(GMaxRHIFeatureLevel);
 	TShaderMapRef<FScreenVS> VertexShader(ShaderMap);
 
-	RHICmdList.SetViewport(0, 0, 0, DestTextureWidth, DestTextureHeight, 0.0f);
+	RHICmdList.SetViewport(0.f, 0.f, 0.f, (float)DestTextureWidth, (float)DestTextureHeight, 0.0f);
 
 	// Perform Writable transitions first
 
@@ -431,7 +431,7 @@ void FSlatePostProcessor::UpsampleRect(FRHICommandListImmediate& RHICmdList, IRe
 			);
 
 		const FVector2f Size(DestRect.Right - DestRect.Left, DestRect.Bottom - DestRect.Top);
-		FShaderParams ShaderParams = FShaderParams::MakePixelShaderParams(FVector4f(Size, SizeUV), (FVector4f)Params.CornerRadius); // LWC_TODO: precision loss
+		FShaderParams ShaderParams = FShaderParams::MakePixelShaderParams(FVector4f(Size, SizeUV), Params.CornerRadius);
 
 
 		PixelShader->SetShaderParams(RHICmdList, ShaderParams);
@@ -488,7 +488,7 @@ static float GetWeight(float Dist, float Strength)
 	return (1.0f / FMath::Sqrt(2 * PI*Strength2))*FMath::Exp(-(Dist*Dist) / (2 * Strength2));
 }
 
-static FVector2D GetWeightAndOffset(float Dist, float Sigma)
+static FVector2f GetWeightAndOffset(float Dist, float Sigma)
 {
 	float Offset1 = Dist;
 	float Weight1 = GetWeight(Offset1, Sigma);
@@ -505,7 +505,7 @@ static FVector2D GetWeightAndOffset(float Dist, float Sigma)
 	}
 
 
-	return FVector2D(TotalWeight, Offset);
+	return FVector2f(TotalWeight, Offset);
 }
 
 static int32 ComputeWeights(int32 KernelSize, float Sigma, TArray<FVector4f>& OutWeightsAndOffsets)
@@ -516,11 +516,11 @@ static int32 ComputeWeights(int32 KernelSize, float Sigma, TArray<FVector4f>& Ou
 
 	OutWeightsAndOffsets.AddUninitialized(NumSamples%2 == 0 ? NumSamples / 2 : NumSamples/2+1);
 
-	OutWeightsAndOffsets[0] = FVector4f(FVector2f(GetWeight(0,Sigma), 0), FVector2f(GetWeightAndOffset(1, Sigma)) );
+	OutWeightsAndOffsets[0] = FVector4f(FVector2f(GetWeight(0,Sigma), 0), GetWeightAndOffset(1, Sigma) );
 	int32 SampleIndex = 1;
 	for (int32 X = 3; X < KernelSize; X += 4)
 	{
-		OutWeightsAndOffsets[SampleIndex] = FVector4f(FVector2f(GetWeightAndOffset(X, Sigma)), FVector2f(GetWeightAndOffset(X + 2, Sigma)));
+		OutWeightsAndOffsets[SampleIndex] = FVector4f(GetWeightAndOffset((float)X, Sigma), GetWeightAndOffset((float)(X + 2), Sigma));
 
 		++SampleIndex;
 	}
diff --git a/Engine/Source/Runtime/SlateRHIRenderer/Private/SlatePostProcessor.h b/Engine/Source/Runtime/SlateRHIRenderer/Private/SlatePostProcessor.h
--- a/Engine/Source/Runtime/SlateRHIRenderer/Private/SlatePostProcessor.h
+++ b/Engine/Source/Runtime/SlateRHIRenderer/Private/SlatePostProcessor.h
@@ -14,7 +14,7 @@ struct FPostProcessRectParams
 	FTextureRHIRef SourceTexture;
 	FSlateRect SourceRect;
 	FSlateRect DestRect;
-	FVector4 CornerRadius;
+	FVector4f CornerRadius;
 	FIntPoint SourceTextureSize;
 	TFunction<void(FRHICommandListImmediate&, FGraphicsPipelineStateInitializer&)> RestoreStateFunc;
 	uint32 StencilRef{};
diff --git a/Engine/Source/Runtime/SlateRHIRenderer/Private/SlateRHIRenderer.cpp b/Engine/Source/Runtime/SlateRHIRenderer/Private/SlateRHIRenderer.cpp
--- a/Engine/Source/Runtime/SlateRHIRenderer/Private/SlateRHIRenderer.cpp
+++ b/Engine/Source/Runtime/SlateRHIRenderer/Private/SlateRHIRenderer.cpp
@@ -966,7 +966,7 @@ void FSlateRHIRenderer::DrawWindow_RenderThread(FRHICommandListImmediate& RHICmd
 			}
 
 			RHICmdList.BeginDrawingViewport(ViewportInfo.ViewportRHI, FTextureRHIRef());
-			RHICmdList.SetViewport(0, 0, 0, ViewportWidth, ViewportHeight, 0.0f);
+			RHICmdList.SetViewport(0.f, 0.f, 0.f, (float)ViewportWidth, (float)ViewportHeight, 0.0f);
 
 			bool bHdrTarget = ViewportInfo.bSceneHDREnabled && !bCompositeUI;
 			RenderSlateBatch(BackBuffer, bClear, bHdrTarget, ViewportInfo, ViewMatrix, BatchData, RHICmdList, ViewportWidth, ViewportHeight, DrawCommandParams, RenderingPolicy, PostProcessBuffer);
@@ -1088,10 +1088,10 @@ void FSlateRHIRenderer::DrawWindow_RenderThread(FRHICommandListImmediate& RHICmd
 
 						RendererModule.DrawRectangle(
 							RHICmdList,
-							0, 0,
-							ViewportWidth, ViewportHeight,
-							0, 0,
-							ViewportWidth, ViewportHeight,
+							0.f, 0.f,
+							(float)ViewportWidth, (float)ViewportHeight,
+							0.f, 0.f,
+							(float)ViewportWidth, (float)ViewportHeight,
 							FIntPoint(ViewportWidth, ViewportHeight),
 							FIntPoint(ViewportWidth, ViewportHeight),
 							VertexShader,
@@ -1139,10 +1139,10 @@ void FSlateRHIRenderer::DrawWindow_RenderThread(FRHICommandListImmediate& RHICmd
 
 				RendererModule.DrawRectangle(
 					RHICmdList,
-					0, 0,
-					ViewportWidth, ViewportHeight,
-					0, 0,
-					ViewportWidth, ViewportHeight,
+					0.f, 0.f,
+					(float)ViewportWidth, (float)ViewportHeight,
+					0.f, 0.f,
+					(float)ViewportWidth, (float)ViewportHeight,
 					FIntPoint(ViewportWidth, ViewportHeight),
 					FIntPoint(ViewportWidth, ViewportHeight),
 					VertexShader,
@@ -1386,8 +1386,8 @@ void FSlateRHIRenderer::DrawWindows_Private(FSlateDrawBuffer& WindowDrawBuffer)
 					Params.bClear = false;
 #endif	
 					Params.Time = FGameTime::CreateDilated(
-						FPlatformTime::Seconds() - GStartTime, FApp::GetDeltaTime(),
-						FApp::GetCurrentTime() - GStartTime, FApp::GetDeltaTime());
+						FPlatformTime::Seconds() - GStartTime, (float)FApp::GetDeltaTime(),
+						FApp::GetCurrentTime() - GStartTime, (float)FApp::GetDeltaTime());
 
 					// Skip the actual draw if we're in a headless execution environment
 					bool bLocalTakingAScreenShot = bTakingAScreenShot;
diff --git a/Engine/Source/Runtime/SlateRHIRenderer/Private/SlateRHIRenderingPolicy.cpp b/Engine/Source/Runtime/SlateRHIRenderer/Private/SlateRHIRenderingPolicy.cpp
--- a/Engine/Source/Runtime/SlateRHIRenderer/Private/SlateRHIRenderingPolicy.cpp
+++ b/Engine/Source/Runtime/SlateRHIRenderer/Private/SlateRHIRenderingPolicy.cpp
@@ -267,7 +267,7 @@ static bool UpdateScissorRect(
 	FTexture2DRHIRef& ColorTarget,
 	FTexture2DRHIRef& DepthStencilTarget,
 	const FSlateClippingState*& LastClippingState,
-	const FVector2D& ViewTranslation2D, 
+	const FVector2f ViewTranslation2D, 
 	bool bSwitchVerticalAxis,
 	FGraphicsPipelineStateInitializer& InGraphicsPSOInit,
 	FSlateStencilClipVertexBuffer& StencilVertexBuffer,
@@ -308,11 +308,11 @@ static bool UpdateScissorRect(
 				const FSlateClippingZone& ScissorRect = ClipState.ScissorRect.GetValue();
 
 				const FIntPoint SizeXY = BackBuffer.GetSizeXY();
-				const FVector2D ViewSize((float) SizeXY.X, (float) SizeXY.Y);
+				const FVector2f ViewSize((float) SizeXY.X, (float) SizeXY.Y);
 
 				// Clamp scissor rect to BackBuffer size
-				const FVector2D TopLeft     = FMath::Min(FMath::Max(FVector2D(ScissorRect.TopLeft)     + ViewTranslation2D, FVector2D(0.0f, 0.0f)), ViewSize);
-				const FVector2D BottomRight = FMath::Min(FMath::Max(FVector2D(ScissorRect.BottomRight) + ViewTranslation2D, FVector2D(0.0f, 0.0f)), ViewSize);
+				const FVector2f TopLeft     = FMath::Min(FMath::Max(ScissorRect.TopLeft     + ViewTranslation2D, FVector2f(0.0f, 0.0f)), ViewSize);
+				const FVector2f BottomRight = FMath::Min(FMath::Max(ScissorRect.BottomRight + ViewTranslation2D, FVector2f(0.0f, 0.0f)), ViewSize);
 				
 				if (bSwitchVerticalAxis)
 				{
@@ -368,7 +368,7 @@ static bool UpdateScissorRect(
 					const FSlateClippingZone& MaskQuad = StencilQuads.Last();
 					const FSlateRect LastStencilBoundingBox = MaskQuad.GetBoundingBox().Round();
 
-					FSlateRect ScissorRect = LastStencilBoundingBox.OffsetBy(ViewTranslation2D);
+					FSlateRect ScissorRect = LastStencilBoundingBox.OffsetBy(FVector2D(ViewTranslation2D));
 
 					// Chosen stencil quad might have some coordinates outside the viewport.
 					// After turning it into a bounding box, this box must be clamped to the current viewport,
@@ -459,7 +459,7 @@ static bool UpdateScissorRect(
 					{
 						const FSlateClippingZone& MaskQuad = StencilQuads[0];
 
-						VertexShader->SetMaskRect(RHICmdList, FVector2D(MaskQuad.TopLeft), FVector2D(MaskQuad.TopRight), FVector2D(MaskQuad.BottomLeft), FVector2D(MaskQuad.BottomRight));
+						VertexShader->SetMaskRect(RHICmdList, MaskQuad.TopLeft, MaskQuad.TopRight, MaskQuad.BottomLeft, MaskQuad.BottomRight);
 
 						RHICmdList.SetStreamSource(0, StencilVertexBuffer.VertexBufferRHI, 0);
 						RHICmdList.DrawPrimitive(0, 2, 1);
@@ -502,7 +502,7 @@ static bool UpdateScissorRect(
 				{
 					const FSlateClippingZone& MaskQuad = StencilQuads[MaskIndex];
 
-					VertexShader->SetMaskRect(RHICmdList, FVector2D(MaskQuad.TopLeft), FVector2D(MaskQuad.TopRight), FVector2D(MaskQuad.BottomLeft), FVector2D(MaskQuad.BottomRight));
+					VertexShader->SetMaskRect(RHICmdList, MaskQuad.TopLeft, MaskQuad.TopRight, MaskQuad.BottomLeft, MaskQuad.BottomRight);
 
 					RHICmdList.SetStreamSource(0, StencilVertexBuffer.VertexBufferRHI, 0);
 					RHICmdList.DrawPrimitive(0, 2, 1);
@@ -593,7 +593,7 @@ void FSlateRHIRenderingPolicy::DrawElements(
 	// and re-gammatization of the resulting blending operation in the 8/10 bit RGBA path.
 	// For Editor running in HDR then the gamma needs to be 2.2 and have a float back buffer format.
 
-	const float EngineGamma = (!GIsEditor && (BackBuffer.GetRenderTargetTexture()->GetFormat() == PF_FloatRGBA) && (Params.bIsHDR==false)) ? 1.0 : GEngine ? GEngine->GetDisplayGamma() : 2.2f;
+	const float EngineGamma = (!GIsEditor && (BackBuffer.GetRenderTargetTexture()->GetFormat() == PF_FloatRGBA) && (Params.bIsHDR==false)) ? 1.0f : GEngine ? GEngine->GetDisplayGamma() : 2.2f;
 	const float DisplayGamma = bGammaCorrect ? EngineGamma : 1.0f;
 	const float DisplayContrast = GSlateContrast;
 
@@ -679,7 +679,7 @@ void FSlateRHIRenderingPolicy::DrawElements(
 	// Disable depth/stencil testing by default
 	GraphicsPSOInit.DepthStencilState = TStaticDepthStencilState<false, CF_Always>::GetRHI();
 
-	FVector2D ViewTranslation2D = Params.ViewOffset;
+	FVector2f ViewTranslation2D = Params.ViewOffset;
 
 	// Draw each element
 #if WITH_SLATE_DEBUGGING
@@ -717,7 +717,7 @@ void FSlateRHIRenderingPolicy::DrawElements(
 			RHICmdList.BeginRenderPass(RPInfo, TEXT("RestartingSlateDrawElements"));
 
 			// Something may have messed with the viewport size so set it back to the full target.
-			RHICmdList.SetViewport(0, 0, 0, BackBuffer.GetSizeXY().X, BackBuffer.GetSizeXY().Y, 0.0f);
+			RHICmdList.SetViewport(0.f, 0.f, 0.f, (float)BackBuffer.GetSizeXY().X, (float)BackBuffer.GetSizeXY().Y, 0.0f);
 
 			// Re-apply render target states to the PSO initializer, since we've changed the depth/stencil target.
 			RHICmdList.ApplyCachedRenderTargets(GraphicsPSOInit);
@@ -1185,10 +1185,10 @@ void FSlateRHIRenderingPolicy::DrawElements(
 
 				FPostProcessRectParams RectParams;
 				RectParams.SourceTexture = PostProcessTexture;
-				RectParams.SourceRect = FSlateRect(0, 0, PostProcessTexture->GetSizeX(), PostProcessTexture->GetSizeY());
+				RectParams.SourceRect = FSlateRect(0.f, 0.f, (float)PostProcessTexture->GetSizeX(), (float)PostProcessTexture->GetSizeY());
 				RectParams.DestRect = FSlateRect(QuadPositionData.X, QuadPositionData.Y, QuadPositionData.Z, QuadPositionData.W);
 				RectParams.SourceTextureSize = PostProcessTexture->GetSizeXY();
-				RectParams.CornerRadius = (FVector4)ShaderParams.PixelParams3;
+				RectParams.CornerRadius = ShaderParams.PixelParams3;
 
 				RectParams.RestoreStateFunc = [&](FRHICommandListImmediate&InRHICmdList, FGraphicsPipelineStateInitializer& InGraphicsPSOInit) {
 					return UpdateScissorRect(
@@ -1263,7 +1263,7 @@ void FSlateRHIRenderingPolicy::DrawElements(
 		RectParams.SourceRect = FSlateRect(0, 0, BackBuffer.GetSizeXY().X, BackBuffer.GetSizeXY().Y);
 		RectParams.DestRect = FSlateRect(0, 0, BackBuffer.GetSizeXY().X, BackBuffer.GetSizeXY().Y);
 		RectParams.SourceTextureSize = BackBuffer.GetSizeXY();
-		RectParams.CornerRadius = FVector4(0, 0, 0, 0);
+		RectParams.CornerRadius = FVector4f(0.f, 0.f, 0.f, 0.f);
 
 		PostProcessor->ColorDeficiency(RHICmdList, RendererModule, RectParams);
 
diff --git a/Engine/Source/Runtime/SlateRHIRenderer/Private/SlateRHIRenderingPolicy.h b/Engine/Source/Runtime/SlateRHIRenderer/Private/SlateRHIRenderingPolicy.h
--- a/Engine/Source/Runtime/SlateRHIRenderer/Private/SlateRHIRenderingPolicy.h
+++ b/Engine/Source/Runtime/SlateRHIRenderer/Private/SlateRHIRenderingPolicy.h
@@ -27,7 +27,7 @@ class FSlateMaterialShaderVS;
 struct FSlateRenderingParams
 {
 	FMatrix44f ViewProjectionMatrix;
-	FVector2D ViewOffset;
+	FVector2f ViewOffset;
 	FGameTime Time;
 	bool bAllowSwitchVerticalAxis;
 	bool bWireFrame;
@@ -35,7 +35,7 @@ struct FSlateRenderingParams
 
 	FSlateRenderingParams(const FMatrix& InViewProjectionMatrix, FGameTime InTime)
 		: ViewProjectionMatrix(InViewProjectionMatrix)
-		, ViewOffset(0, 0)
+		, ViewOffset(0.f, 0.f)
 		, Time(InTime)
 		, bAllowSwitchVerticalAxis(true)
 		, bWireFrame(false)
diff --git a/Engine/Source/Runtime/SlateRHIRenderer/Private/SlateRHIResourceManager.cpp b/Engine/Source/Runtime/SlateRHIRenderer/Private/SlateRHIResourceManager.cpp
--- a/Engine/Source/Runtime/SlateRHIRenderer/Private/SlateRHIResourceManager.cpp
+++ b/Engine/Source/Runtime/SlateRHIRenderer/Private/SlateRHIResourceManager.cpp
@@ -24,6 +24,7 @@
 #include "SlateMaterialResource.h"
 #include "Slate/SlateTextureAtlasInterface.h"
 #include "SlateAtlasedTextureResource.h"
+#include "Types/SlateVector2.h"
 #include "ImageUtils.h"
 #include "Async/ParallelFor.h"
 #include "Rendering/SlateVectorGraphicsCache.h"
@@ -397,7 +398,7 @@ FAtlasSlotInfo FSlateRHIResourceManager::GetAtlasSlotInfoAtPosition(FIntPoint In
 		const FAtlasedTextureSlot* Slot = Atlas->GetSlotAtPosition(InPosition);
 		if (Slot)
 		{
-			NewInfo.AtlasSlotRect = FSlateRect(FVector2D(Slot->X, Slot->Y), FVector2D(Slot->X + Slot->Width, Slot->Y + Slot->Height));
+			NewInfo.AtlasSlotRect = FSlateRect(FVector2f((float)Slot->X, (float)Slot->Y), FVector2f((float)(Slot->X + Slot->Width), (float)(Slot->Y + Slot->Height)));
 
 			NewInfo.TextureName = AtlasDebugData.FindRef(Slot);
 
@@ -797,7 +798,7 @@ TSharedPtr<FSlateUTextureResource> FSlateRHIResourceManager::MakeDynamicUTexture
 			TextureResource = MakeShareable(new FSlateUTextureResource(InTextureObject));
 		}
 
-		TextureResource->Proxy->ActualSize = FIntPoint(InTextureObject->GetSurfaceWidth(), InTextureObject->GetSurfaceHeight());
+		TextureResource->Proxy->ActualSize = FIntPoint(FMath::TruncToInt32(InTextureObject->GetSurfaceWidth()), FMath::TruncToInt32(InTextureObject->GetSurfaceHeight()));
 	}
 	else
 	{
@@ -909,7 +910,7 @@ FSlateMaterialResource* FSlateRHIResourceManager::GetMaterialResource(const UObj
 
 	const UMaterialInterface* Material = CastChecked<UMaterialInterface>(InMaterial);
 
-	FVector2D ImageSize = InBrush ? InBrush->ImageSize : FVector2D::ZeroVector;
+	FVector2f ImageSize = InBrush ? UE::Slate::CastToVector2f(InBrush->ImageSize) : FVector2f::ZeroVector;
 	FMaterialKey Key(Material, ImageSize, InMaskKey);
 
 	TSharedPtr<FSlateMaterialResource> MaterialResource = DynamicResourceMap.GetMaterialResource(Key);
@@ -989,7 +990,7 @@ void FSlateRHIResourceManager::ReleaseDynamicResource( const FSlateBrush& InBrus
 			{
 				UMaterialInterface* Material = Cast<UMaterialInterface>(ResourceObject);
 
-				FMaterialKey Key(Material, InBrush.ImageSize, 0);
+				FMaterialKey Key(Material, UE::Slate::CastToVector2f(InBrush.ImageSize), 0);
 
 				TSharedPtr<FSlateMaterialResource> MaterialResource = DynamicResourceMap.GetMaterialResource(Key);
 				
diff --git a/Engine/Source/Runtime/SlateRHIRenderer/Private/SlateRHIResourceManager.h b/Engine/Source/Runtime/SlateRHIRenderer/Private/SlateRHIResourceManager.h
--- a/Engine/Source/Runtime/SlateRHIRenderer/Private/SlateRHIResourceManager.h
+++ b/Engine/Source/Runtime/SlateRHIRenderer/Private/SlateRHIResourceManager.h
@@ -31,10 +31,10 @@ class FSlateVectorGraphicsCache;
 struct FMaterialKey
 {
 	TWeakObjectPtr<const UMaterialInterface> Material;
-	const FVector2D ImageSize;
+	const FVector2f ImageSize;
 	int32 MaskKey;
 
-	FMaterialKey(const UMaterialInterface* InMaterial, const FVector2D& InImageSize, int32 InMaskKey)
+	FMaterialKey(const UMaterialInterface* InMaterial, const FVector2f InImageSize, int32 InMaskKey)
 		: Material(InMaterial)
 		, ImageSize(InImageSize)
 		, MaskKey(InMaskKey)
diff --git a/Engine/Source/Runtime/SlateRHIRenderer/Private/SlateShaders.cpp b/Engine/Source/Runtime/SlateRHIRenderer/Private/SlateShaders.cpp
--- a/Engine/Source/Runtime/SlateRHIRenderer/Private/SlateShaders.cpp
+++ b/Engine/Source/Runtime/SlateRHIRenderer/Private/SlateShaders.cpp
@@ -73,7 +73,7 @@ TGlobalResource<FSlateMaskingVertexDeclaration> GSlateMaskingVertexDeclaration;
 void FSlateVertexDeclaration::InitRHI()
 {
 	FVertexDeclarationElementList Elements;
-	uint32 Stride = sizeof(FSlateVertex);
+	uint16 Stride = sizeof(FSlateVertex);
 	Elements.Add(FVertexElement(0, STRUCT_OFFSET(FSlateVertex, TexCoords), VET_Float4, 0, Stride));
 	Elements.Add(FVertexElement(0, STRUCT_OFFSET(FSlateVertex, MaterialTexCoords), VET_Float2, 1, Stride));
 	Elements.Add(FVertexElement(0, STRUCT_OFFSET(FSlateVertex, Position), VET_Float2, 2, Stride));
@@ -96,7 +96,7 @@ void FSlateVertexDeclaration::ReleaseRHI()
 void FSlateInstancedVertexDeclaration::InitRHI()
 {
 	FVertexDeclarationElementList Elements;
-	uint32 Stride = sizeof(FSlateVertex);
+	uint16 Stride = sizeof(FSlateVertex);
 	Elements.Add(FVertexElement(0, STRUCT_OFFSET(FSlateVertex, TexCoords), VET_Float4, 0, Stride));
 	Elements.Add(FVertexElement(0, STRUCT_OFFSET(FSlateVertex, MaterialTexCoords), VET_Float2, 1, Stride));
 	Elements.Add(FVertexElement(0, STRUCT_OFFSET(FSlateVertex, Position), VET_Float2, 2, Stride));
@@ -120,7 +120,7 @@ void FSlateElementPS::ModifyCompilationEnvironment(const FGlobalShaderPermutatio
 void FSlateMaskingVertexDeclaration::InitRHI()
 {
 	FVertexDeclarationElementList Elements;
-	uint32 Stride = sizeof(uint32);
+	uint16 Stride = sizeof(uint32);
 	Elements.Add(FVertexElement(0, 0, VET_UByte4, 0, Stride));
 
 	VertexDeclarationRHI = PipelineStateCache::GetOrCreateVertexDeclaration(Elements);
@@ -194,10 +194,9 @@ void FSlateMaskingVS::SetVerticalAxisMultiplier(FRHICommandList& RHICmdList, flo
 	SetShaderValue(RHICmdList, RHICmdList.GetBoundVertexShader(), SwitchVerticalAxisMultiplier, InMultiplier );
 }
 
-void FSlateMaskingVS::SetMaskRect(FRHICommandList& RHICmdList, const FVector2D& TopLeft, const FVector2D& TopRight, const FVector2D& BotLeft, const FVector2D& BotRight)
+void FSlateMaskingVS::SetMaskRect(FRHICommandList& RHICmdList, const FVector2f TopLeft, const FVector2f TopRight, const FVector2f BotLeft, const FVector2f BotRight)
 {
-	//FVector4f MaskRectVal[4] = { FVector4f(TopLeft, FVector2D::ZeroVector), FVector4f(TopRight, FVector2D::ZeroVector), FVector4f(BotLeft, FVector2D::ZeroVector), FVector4f(BotRight, FVector2D::ZeroVector) };
-	FVector4f MaskRectVal[2] = { FVector4f(FVector2f(TopLeft), FVector2f(TopRight)), FVector4f(FVector2f(BotLeft), FVector2f(BotRight)) };	// LWC_TODO: Precision loss
+	FVector4f MaskRectVal[2] = { FVector4f(TopLeft, TopRight), FVector4f(BotLeft, BotRight) };
 
 	SetShaderValue(RHICmdList, RHICmdList.GetBoundVertexShader(), MaskRect, MaskRectVal);
 }
diff --git a/Engine/Source/Runtime/SlateRHIRenderer/Private/SlateShaders.h b/Engine/Source/Runtime/SlateRHIRenderer/Private/SlateShaders.h
--- a/Engine/Source/Runtime/SlateRHIRenderer/Private/SlateShaders.h
+++ b/Engine/Source/Runtime/SlateRHIRenderer/Private/SlateShaders.h
@@ -192,8 +192,8 @@ public:
 	 */
 	void SetShaderParams(FRHICommandList& RHICmdList, const FShaderParams& InShaderParams)
 	{
-		SetShaderValue(RHICmdList, RHICmdList.GetBoundPixelShader(), ShaderParams, (FVector4f)InShaderParams.PixelParams);
-		SetShaderValue(RHICmdList, RHICmdList.GetBoundPixelShader(), ShaderParams2, (FVector4f)InShaderParams.PixelParams2);
+		SetShaderValue(RHICmdList, RHICmdList.GetBoundPixelShader(), ShaderParams, InShaderParams.PixelParams);
+		SetShaderValue(RHICmdList, RHICmdList.GetBoundPixelShader(), ShaderParams2, InShaderParams.PixelParams2);
 	}
 
 	/**
@@ -346,9 +346,9 @@ public:
 		UVBounds.Bind(Initializer.ParameterMap, TEXT("UVBounds"));
 	}
 
-	void SetBufferSizeAndDirection(FRHICommandList& RHICmdList, const FVector2D& InBufferSize, const FVector2D& InDir)
+	void SetBufferSizeAndDirection(FRHICommandList& RHICmdList, const FVector2f InBufferSize, const FVector2f InDir)
 	{
-		SetShaderValue(RHICmdList, RHICmdList.GetBoundPixelShader(), BufferSizeAndDirection, FVector4f(FVector2f(InBufferSize), FVector2f(InDir)));	// LWC_TODO: Precision loss
+		SetShaderValue(RHICmdList, RHICmdList.GetBoundPixelShader(), BufferSizeAndDirection, FVector4f(InBufferSize, InDir));
 	}
 
 	void SetWeightsAndOffsets(FRHICommandList& RHICmdList, const TArray<FVector4f>& InWeightsAndOffsets, int32 NumSamples )
@@ -502,7 +502,7 @@ public:
 	/**
 	 * Sets the mask rect positions
 	 */
-	void SetMaskRect(FRHICommandList& RHICmdList, const FVector2D& TopLeft, const FVector2D& TopRight, const FVector2D& BotLeft, const FVector2D& BotRight);
+	void SetMaskRect(FRHICommandList& RHICmdList, const FVector2f TopLeft, const FVector2f TopRight, const FVector2f BotLeft, const FVector2f BotRight);
 
 	//virtual bool Serialize(FArchive& Ar) override;
 
