fixes for background blur with HDR
SlatePostProcessor was hard-coding the pixel format to 8bpc, it now uses the same format as the source texture.
Have SlateRHIRenderer render HDR scene first: this way, background blur does have the scene available to can mix it with the UI rendertarget
Composite UI with 3D scene before blur happens so that blur only operates on a single texture
Convert PQ scene color to linear 11-11-10 bits because doing PQ->Linear conversion in the blur shader was pretty expensive, and using a LUT didn't help as we were getting memory bound. 16bpc was also introducing a significant overhead during the blur pass
Because we now use 11-11-10, there's some banding happening. This can be acceptable on platforms that support rendertarget writemask, because we discard pixels untouched by the UI. On platforms that are not supporting it, this introduced some banding. Adding an RT mask fixes & testing the mask + the UI alpha channel now allows to discard pixels and keep the original value intact if untouched by the UI.

Fixed stencil masking for background blur as it would call UpdateScissorRect before doing the upsampling, and would not work with up-to-date stencil ref

diff --git a/Engine/Shaders/Private/CompositeUIPixelShader.usf b/Engine/Shaders/Private/CompositeUIPixelShader.usf
--- a/Engine/Shaders/Private/CompositeUIPixelShader.usf
+++ b/Engine/Shaders/Private/CompositeUIPixelShader.usf
@@ -62,9 +62,7 @@ void GenerateLUTPS(
 }
 
 ///////////////////////////////////////////////////////////////////////////////////////
-#if PLATFORM_SUPPORTS_RENDERTARGET_WRITE_MASK
 Texture2D<uint> UIWriteMaskTexture;
-#endif
 
 Texture2D UITexture;
 SamplerState UISampler;
@@ -93,6 +91,9 @@ float3 LinearizeColor(float3 EncodedColor)
 	#endif
 }
 
+static const float ScRGBScaleFactor = 80.0f; // used to convert between ScRGB and nits
+static const float UIBrightnessInNits = 300.0f; // todo stop hardcoding this value and try to fetch SDR/HDR brightness balance instead (automatically or though calibration step)
+
 float3 ColorLookupTable(float3 LinearColor)
 {
 	float3 LUTEncodedColor = LinearToST2084(LinearColor);
@@ -100,6 +101,17 @@ float3 ColorLookupTable(float3 LinearColor)
 	return ColorSpaceLUT.SampleLevel(ColorSpaceLUTSampler, UVW, 0).rgb * LUT_SCALING;
 }
 
+void HDRBackBufferConvertPS(
+	FScreenVertexOutput Input,
+	out float3 OutColor : SV_Target0
+	)
+{
+	OutColor = 0;
+	float3 SceneColor = Texture2DSample(SceneTexture, SceneSampler, Input.UV).xyz ;
+	SceneColor = ST2084ToLinear( SceneColor ) / ScRGBScaleFactor; 
+	OutColor.rgb = SceneColor;
+}
+
 // Given an input using Rec709 primaries and a sRGB tranfer function where
 // 1.0 maps to paperwhitenits defined in editor. Remap to Rec2020 with PQ transfer function.
 void HDREditorConvert(
@@ -134,15 +146,24 @@ void Main(
 	{
 		discard;
 	}
+#else
+	uint CompositeUIMask = UIWriteMaskTexture.Load(uint3(Input.Position.xy,0));
 #endif
 
 	// In order to get consistent blending between HDR PQ & ScRGB, we convert UI color in both cases to rec2020, perform the blending in this color space, and
 	// move back to the output device space. Rec2020 is chosen over ScRGB because this is the way blending is currently being done in engine mode, and contains only positive values
 	const float3x3 sRGB_2_Rec2020 = mul(XYZ_2_Rec2020_MAT, sRGB_2_XYZ_MAT);
-	static const float PaperWhiteInNits = 80.0f; // used to convert between ScRGB and nits
-	static const float UIBrightnessInNits = 300.0f; // todo stop hardcoding this value and try to fetch SDR/HDR brightness balance instead (automatically or though calibration step)
 
 	float4 UIColor = Texture2DSample(UITexture, UISampler, Input.UV);
+
+#if !PLATFORM_SUPPORTS_RENDERTARGET_WRITE_MASK
+	BRANCH
+	if (CompositeUIMask == 0 && UIColor.a < 1e-6f)
+	{
+		discard;
+	}
+#endif
+
 	{
 		UIColor.xyz = LinearizeColor(UIColor.xyz);
 #if USE_LUT_FOR_UI
@@ -162,9 +183,10 @@ void Main(
 
 	float3 SceneColor = Texture2DSample(SceneTexture, SceneSampler, Input.UV).xyz;
 #if !SCRGB_ENCODING
-	SceneColor = ST2084ToLinear(SceneColor);
+    // We don't use ST2084ToLinear here because HDRBackBufferConvertPS performed the conversion before
+	SceneColor.xyz = SceneColor.xyz * ScRGBScaleFactor;
 #else
-	SceneColor.xyz = mul(sRGB_2_Rec2020, SceneColor.xyz) * PaperWhiteInNits;
+	SceneColor.xyz = mul(sRGB_2_Rec2020, SceneColor.xyz) * ScRGBScaleFactor;
 #endif
 
 	BRANCH
@@ -193,10 +215,9 @@ void Main(
 	OutColor.xyz = LinearToST2084(OutColor.xyz);
 #else
 	const float3x3 Rec2020_2_sRGB = mul(XYZ_2_sRGB_MAT, Rec2020_2_XYZ_MAT);
-	OutColor.xyz = mul(Rec2020_2_sRGB, OutColor.xyz / PaperWhiteInNits);
+	OutColor.xyz = mul(Rec2020_2_sRGB, OutColor.xyz / ScRGBScaleFactor);
 #endif
 }
-
 #if DISPLAY_MAPPING_PS
 float4x4 TextureToOutputGamutMatrix;
 
@@ -222,4 +243,41 @@ void DisplayMappingPS(
 	}
 }
 
-#endif
\ No newline at end of file
+#endif
+
+#if BLIT_UI_TO_HDR
+
+float4x4 SrgbToOutputMatrix;
+
+void BlitUIToHDRPS(
+	FScreenVertexOutput Input,
+	out float4 OutColor : SV_Target0
+	)
+{
+#if PLATFORM_SUPPORTS_RENDERTARGET_WRITE_MASK
+	uint CompositeUIMask = DecodeRTWriteMask(uint2(Input.Position.xy), UIWriteMaskTexture, 1);
+	BRANCH
+	if (CompositeUIMask == 0)
+	{
+		discard;
+	}
+#endif
+
+	float4 UIColor = Texture2DSample(UITexture, UISampler, Input.UV);
+
+#if !PLATFORM_SUPPORTS_RENDERTARGET_WRITE_MASK
+	BRANCH
+	if (UIColor.a < 1e-3f)
+	{
+		discard;
+	}
+#endif
+
+	UIColor.rgb = LinearizeColor(UIColor.rgb) * UILevel * UIBrightnessInNits / ScRGBScaleFactor * UIColor.a;
+	UIColor.rgb = mul((float3x3)SrgbToOutputMatrix, UIColor.rgb);
+
+	OutColor.rgb = UIColor.rgb;
+	OutColor.a = UIColor.a;
+}
+
+#endif
diff --git a/Engine/Shaders/Private/SlatePostProcessPixelShader.usf b/Engine/Shaders/Private/SlatePostProcessPixelShader.usf
--- a/Engine/Shaders/Private/SlatePostProcessPixelShader.usf
+++ b/Engine/Shaders/Private/SlatePostProcessPixelShader.usf
@@ -114,13 +114,30 @@ float2 Map(float2 value, float2 min1, float2 max1, float2 min2, float2 max2)
 {
   return min2 + (value - min1) * (max2 - min2) / (max1 - min1);
 }
-
-float4 UpsampleMain(FScreenVertexOutput Input) : SV_Target0
+ 
+void UpsampleMain(FScreenVertexOutput Input
+	, out float4 Color0 : SV_Target0
+#if OUTPUT_TO_UI_TARGET
+	, out float4 Color1 : SV_Target1
+#if !PLATFORM_SUPPORTS_RENDERTARGET_WRITE_MASK
+	, out uint Color2 : SV_Target2
+#endif
+#endif
+)
 {
 	const float2 LocalSize = ShaderParams.xy;
 	const float2 UV = Input.UV;
 	const float Thickness = 0;
 	const float4 CornerRadii = ShaderParams2;
 	const float4 FillColor = Texture2DSample(ElementTexture, ElementTextureSampler, Input.UV);
-	return GetRoundedBoxElementColorInternal(LocalSize, Map(UV, float2(0,0), ShaderParams.zw, float2(0,0),float2(1,1)), Thickness, CornerRadii, FillColor, FillColor);
+	Color0 = GetRoundedBoxElementColorInternal(LocalSize, Map(UV, float2(0,0), ShaderParams.zw, float2(0,0),float2(1,1)), Thickness, CornerRadii, FillColor, FillColor);
+#if OUTPUT_TO_UI_TARGET
+	// still needed in case compositing needs PLATFORM_SUPPORTS_RENDERTARGET_WRITE_MASK
+	Color1 = float4(0, 0, 0, 0);
+
+#if !PLATFORM_SUPPORTS_RENDERTARGET_WRITE_MASK
+	Color2 = 255;
+#endif
+
+#endif
 }
\ No newline at end of file
diff --git a/Engine/Source/Runtime/RenderCore/Private/RenderCore.cpp b/Engine/Source/Runtime/RenderCore/Private/RenderCore.cpp
--- a/Engine/Source/Runtime/RenderCore/Private/RenderCore.cpp
+++ b/Engine/Source/Runtime/RenderCore/Private/RenderCore.cpp
@@ -552,4 +552,73 @@ RENDERCORE_API void HDRConfigureCVars(bool bIsHDREnabled, uint32 DisplayNits, bo
 
 	CVarDisplayOutputDevice->Set((int32)OutputDevice, ECVF_SetByDeviceProfile);
 	CVarDisplayColorGamut->Set((int32)ColorGamut, ECVF_SetByDeviceProfile);
-}
\ No newline at end of file
+}
+
+RENDERCORE_API FMatrix44f GamutToXYZMatrix(EDisplayColorGamut ColorGamut)
+{
+	static const FMatrix44f sRGB_2_XYZ_MAT(
+		FVector3f(0.4124564, 0.3575761, 0.1804375),
+		FVector3f(0.2126729, 0.7151522, 0.0721750),
+		FVector3f(0.0193339, 0.1191920, 0.9503041),
+		FVector3f(0, 0, 0)
+	);
+
+	static const FMatrix44f Rec2020_2_XYZ_MAT(
+		FVector3f(0.6369736, 0.1446172, 0.1688585),
+		FVector3f(0.2627066, 0.6779996, 0.0592938),
+		FVector3f(0.0000000, 0.0280728, 1.0608437),
+		FVector3f(0, 0, 0)
+	);
+
+	static const FMatrix44f P3D65_2_XYZ_MAT(
+		FVector3f(0.4865906, 0.2656683, 0.1981905),
+		FVector3f(0.2289838, 0.6917402, 0.0792762),
+		FVector3f(0.0000000, 0.0451135, 1.0438031),
+		FVector3f(0, 0, 0)
+	);
+	switch (ColorGamut)
+	{
+	case EDisplayColorGamut::sRGB_D65: return sRGB_2_XYZ_MAT;
+	case EDisplayColorGamut::Rec2020_D65: return Rec2020_2_XYZ_MAT;
+	case EDisplayColorGamut::DCIP3_D65: return P3D65_2_XYZ_MAT;
+	default:
+		checkNoEntry();
+		return FMatrix44f::Identity;
+	}
+
+}
+
+RENDERCORE_API FMatrix44f XYZToGamutMatrix(EDisplayColorGamut ColorGamut)
+{
+	static const FMatrix44f XYZ_2_sRGB_MAT(
+		FVector3f(3.2409699419, -1.5373831776, -0.4986107603),
+		FVector3f(-0.9692436363, 1.8759675015, 0.0415550574),
+		FVector3f(0.0556300797, -0.2039769589, 1.0569715142),
+		FVector3f(0, 0, 0)
+	);
+
+	static const FMatrix44f XYZ_2_Rec2020_MAT(
+		FVector3f(1.7166084, -0.3556621, -0.2533601),
+		FVector3f(-0.6666829, 1.6164776, 0.0157685),
+		FVector3f(0.0176422, -0.0427763, 0.94222867),
+		FVector3f(0, 0, 0)
+	);
+
+	static const FMatrix44f XYZ_2_P3D65_MAT(
+		FVector3f(2.4933963, -0.9313459, -0.4026945),
+		FVector3f(-0.8294868, 1.7626597, 0.0236246),
+		FVector3f(0.0358507, -0.0761827, 0.9570140),
+		FVector3f(0, 0, 0)
+	);
+
+	switch (ColorGamut)
+	{
+	case EDisplayColorGamut::sRGB_D65: return XYZ_2_sRGB_MAT;
+	case EDisplayColorGamut::Rec2020_D65: return XYZ_2_Rec2020_MAT;
+	case EDisplayColorGamut::DCIP3_D65: return XYZ_2_P3D65_MAT;
+	default:
+		checkNoEntry();
+		return FMatrix44f::Identity;
+	}
+
+}
diff --git a/Engine/Source/Runtime/RenderCore/Public/HDRHelper.h b/Engine/Source/Runtime/RenderCore/Public/HDRHelper.h
--- a/Engine/Source/Runtime/RenderCore/Public/HDRHelper.h
+++ b/Engine/Source/Runtime/RenderCore/Public/HDRHelper.h
@@ -12,6 +12,8 @@ RENDERCORE_API EDisplayOutputFormat HDRGetDefaultDisplayOutputFormat();
 RENDERCORE_API EDisplayColorGamut HDRGetDefaultDisplayColorGamut();
 RENDERCORE_API void HDRAddCustomMetaData(void* OSWindow, EDisplayOutputFormat DisplayOutputFormat, EDisplayColorGamut DisplayColorGamut, bool bHDREnabled);
 RENDERCORE_API void HDRRemoveCustomMetaData(void* OSWindow);
+RENDERCORE_API FMatrix44f GamutToXYZMatrix(EDisplayColorGamut ColorGamut);
+RENDERCORE_API FMatrix44f XYZToGamutMatrix(EDisplayColorGamut ColorGamut);
 
 
 struct FDisplayInformation
diff --git a/Engine/Source/Runtime/SlateRHIRenderer/Private/SlatePostProcessResource.cpp b/Engine/Source/Runtime/SlateRHIRenderer/Private/SlatePostProcessResource.cpp
--- a/Engine/Source/Runtime/SlateRHIRenderer/Private/SlatePostProcessResource.cpp
+++ b/Engine/Source/Runtime/SlateRHIRenderer/Private/SlatePostProcessResource.cpp
@@ -6,7 +6,8 @@
 DECLARE_MEMORY_STAT(TEXT("PostProcess RenderTargets"), STAT_SLATEPPRenderTargetMem, STATGROUP_SlateMemory);
 
 FSlatePostProcessResource::FSlatePostProcessResource(int32 InRenderTargetCount)
-	: RenderTargetSize(FIntPoint::ZeroValue)
+	: PixelFormat(PF_Unknown)
+	, RenderTargetSize(FIntPoint::ZeroValue)
 	, RenderTargetCount(InRenderTargetCount)
 {
 }
@@ -16,9 +17,9 @@ FSlatePostProcessResource::~FSlatePostProcessResource()
 
 }
 
-void FSlatePostProcessResource::Update(const FIntPoint& NewSize)
+void FSlatePostProcessResource::Update(const FIntPoint& NewSize, const FTextureRHIRef& SourceTexture)
 {
-	if(NewSize.X > RenderTargetSize.X || NewSize.Y > RenderTargetSize.Y || RenderTargetSize == FIntPoint::ZeroValue || RenderTargets.Num() == 0 )
+	if(NewSize.X > RenderTargetSize.X || NewSize.Y > RenderTargetSize.Y || RenderTargetSize == FIntPoint::ZeroValue || RenderTargets.Num() == 0 || PixelFormat != SourceTexture->GetDesc().Format)
 	{
 		if(!IsInitialized())
 		{
@@ -26,18 +27,18 @@ void FSlatePostProcessResource::Update(const FIntPoint& NewSize)
 		}
 
 		FIntPoint NewMaxSize(FMath::Max(NewSize.X, RenderTargetSize.X), FMath::Max(NewSize.Y, RenderTargetSize.Y));
-		ResizeTargets(NewMaxSize);
+		ResizeTargets(NewMaxSize, SourceTexture);
 	}
 }
 
-void FSlatePostProcessResource::ResizeTargets(const FIntPoint& NewSize)
+void FSlatePostProcessResource::ResizeTargets(const FIntPoint& NewSize, const FTextureRHIRef& SourceTexture)
 {
 	check(IsInRenderingThread());
 
 	RenderTargets.Empty();
 
 	RenderTargetSize = NewSize;
-	PixelFormat = PF_B8G8R8A8;
+	PixelFormat = SourceTexture->GetDesc().Format;
 	if (RenderTargetSize.X > 0 && RenderTargetSize.Y > 0)
 	{
 		for (int32 TexIndex = 0; TexIndex < RenderTargetCount; ++TexIndex)
@@ -49,7 +50,7 @@ void FSlatePostProcessResource::ResizeTargets(const FIntPoint& NewSize)
 				.SetFlags(ETextureCreateFlags::RenderTargetable | ETextureCreateFlags::ShaderResource)
 				.SetInitialState(ERHIAccess::SRVMask);
 
-			FTexture2DRHIRef RenderTargetTextureRHI = RHICreateTexture(Desc);
+			FTextureRHIRef RenderTargetTextureRHI = RHICreateTexture(Desc);
 
 			RenderTargets.Add(RenderTargetTextureRHI);
 		}
diff --git a/Engine/Source/Runtime/SlateRHIRenderer/Private/SlatePostProcessResource.h b/Engine/Source/Runtime/SlateRHIRenderer/Private/SlatePostProcessResource.h
--- a/Engine/Source/Runtime/SlateRHIRenderer/Private/SlatePostProcessResource.h
+++ b/Engine/Source/Runtime/SlateRHIRenderer/Private/SlatePostProcessResource.h
@@ -22,7 +22,7 @@ public:
 	}
 
 	/** Performs per frame updates to this resource */
-	void Update(const FIntPoint& NewSize);
+	void Update(const FIntPoint& NewSize, const FTextureRHIRef& SourceTexture);
 
 	void CleanUp();
 
@@ -38,7 +38,7 @@ public:
 
 private:
 	/** Resizes targets to the new size */
-	void ResizeTargets(const FIntPoint& NewSize);
+	void ResizeTargets(const FIntPoint& NewSize, const FTextureRHIRef& SourceTexture);
 
 private:
 	TArray<FTexture2DRHIRef, TInlineAllocator<2>> RenderTargets;
diff --git a/Engine/Source/Runtime/SlateRHIRenderer/Private/SlatePostProcessor.cpp b/Engine/Source/Runtime/SlateRHIRenderer/Private/SlatePostProcessor.cpp
--- a/Engine/Source/Runtime/SlateRHIRenderer/Private/SlatePostProcessor.cpp
+++ b/Engine/Source/Runtime/SlateRHIRenderer/Private/SlatePostProcessor.cpp
@@ -9,6 +9,8 @@
 #include "StaticBoundShaderState.h"
 #include "PipelineStateCache.h"
 #include "CommonRenderResources.h"
+#include "HDRHelper.h"
+#include "RendererUtils.h"
 
 DECLARE_CYCLE_STAT(TEXT("Slate PostProcessing RT"), STAT_SlatePostProcessingRTTime, STATGROUP_Slate);
 DECLARE_CYCLE_STAT(TEXT("Slate ColorDeficiency RT"), STAT_SlateColorDeficiencyRTTime, STATGROUP_Slate);
@@ -26,6 +28,151 @@ FSlatePostProcessor::~FSlatePostProcessor()
 	IntermediateTargets->CleanUp();
 }
 
+
+// Pixel shader to composite UI over HDR buffer
+class FBlitUIToHDRPS : public FGlobalShader
+{
+	DECLARE_SHADER_TYPE(FBlitUIToHDRPS, Global);
+public:
+
+	static bool ShouldCompilePermutation(const FGlobalShaderPermutationParameters& Parameters)
+	{
+		return IsFeatureLevelSupported(Parameters.Platform, ERHIFeatureLevel::SM5) && (RHISupportsGeometryShaders(Parameters.Platform) || RHISupportsVertexShaderLayer(Parameters.Platform));
+	}
+
+	FBlitUIToHDRPS(const ShaderMetaType::CompiledShaderInitializerType& Initializer) :
+		FGlobalShader(Initializer)
+	{
+		UITexture.Bind(Initializer.ParameterMap, TEXT("UITexture"));
+		UIWriteMaskTexture.Bind(Initializer.ParameterMap, TEXT("UIWriteMaskTexture"));
+		UISampler.Bind(Initializer.ParameterMap, TEXT("UISampler"));
+		UILevel.Bind(Initializer.ParameterMap, TEXT("UILevel"));
+		SrgbToOutputMatrix.Bind(Initializer.ParameterMap, TEXT("SrgbToOutputMatrix"));
+	}
+	FBlitUIToHDRPS() {}
+
+	void SetParameters(FRHICommandList& RHICmdList, FRHITexture* UITextureRHI, FRHITexture* UITextureWriteMaskRHI, const FMatrix44f& InSrgbToOutputMatrix)
+	{
+		SetTextureParameter(RHICmdList, RHICmdList.GetBoundPixelShader(), UITexture, UISampler, TStaticSamplerState<SF_Point>::GetRHI(), UITextureRHI);
+		static auto CVarHDRUILevel = IConsoleManager::Get().FindConsoleVariable(TEXT("r.HDR.UI.Level"));
+		SetShaderValue(RHICmdList, RHICmdList.GetBoundPixelShader(), UILevel, CVarHDRUILevel ? CVarHDRUILevel->GetFloat() : 1.0f);
+		SetShaderValue(RHICmdList, RHICmdList.GetBoundPixelShader(), SrgbToOutputMatrix, InSrgbToOutputMatrix);
+		if (UITextureWriteMaskRHI != nullptr)
+		{
+			SetTextureParameter(RHICmdList, RHICmdList.GetBoundPixelShader(), UIWriteMaskTexture, UITextureWriteMaskRHI);
+		}
+	}
+
+	static void ModifyCompilationEnvironment(const FGlobalShaderPermutationParameters& Parameters, FShaderCompilerEnvironment& OutEnvironment)
+	{
+		FGlobalShader::ModifyCompilationEnvironment(Parameters, OutEnvironment);
+		OutEnvironment.SetDefine(TEXT("BLIT_UI_TO_HDR"), 1);
+	}
+
+	static const TCHAR* GetSourceFilename()
+	{
+		return TEXT("/Engine/Private/CompositeUIPixelShader.usf");
+	}
+
+	static const TCHAR* GetFunctionName()
+	{
+		return TEXT("BlitUIToHDRPS");
+	}
+
+private:
+	LAYOUT_FIELD(FShaderResourceParameter, UITexture);
+	LAYOUT_FIELD(FShaderResourceParameter, UISampler);
+	LAYOUT_FIELD(FShaderResourceParameter, UIWriteMaskTexture);
+	LAYOUT_FIELD(FShaderParameter, SrgbToOutputMatrix);
+	LAYOUT_FIELD(FShaderParameter, UILevel);
+};
+
+IMPLEMENT_SHADER_TYPE(, FBlitUIToHDRPS, FBlitUIToHDRPS::GetSourceFilename(), FBlitUIToHDRPS::GetFunctionName(), SF_Pixel);
+
+static void BlitUIToHDRScene(FRHICommandListImmediate& RHICmdList, IRendererModule& RendererModule, const FPostProcessRectParams& RectParams)
+{
+	SCOPED_DRAW_EVENT(RHICmdList, SlatePostProcessBlitUIToHDR);
+
+	FRHITexture* SourceTexture = RectParams.UITarget->GetRHI();
+
+	RHICmdList.Transition(FRHITransitionInfo(SourceTexture, ERHIAccess::Unknown, ERHIAccess::SRVMask));
+	TRefCountPtr<IPooledRenderTarget> UITargetRTMask;
+	if (RHISupportsRenderTargetWriteMask(GMaxRHIShaderPlatform))
+	{
+		const auto FeatureLevel = GMaxRHIFeatureLevel;
+		auto ShaderMap = GetGlobalShaderMap(FeatureLevel);
+
+		IPooledRenderTarget* RenderTargets[] = { RectParams.UITarget.GetReference() };
+		FRenderTargetWriteMask::Decode(RHICmdList, ShaderMap, RenderTargets, UITargetRTMask, TexCreate_None, TEXT("UIRTWriteMask"));
+	}
+
+	// Source is the viewport.  This is the width and height of the viewport backbuffer
+	const int32 SrcTextureWidth = RectParams.SourceTextureSize.X;
+	const int32 SrcTextureHeight = RectParams.SourceTextureSize.Y;
+
+	// Rect of the viewport
+	const FSlateRect& SourceRect = RectParams.SourceRect;
+
+	// Rect of the final destination post process effect (not downsample rect).  This is the area we sample from
+	const FSlateRect& DestRect = RectParams.DestRect;
+
+	FGlobalShaderMap* ShaderMap = GetGlobalShaderMap(GMaxRHIFeatureLevel);
+	TShaderMapRef<FScreenVS> VertexShader(ShaderMap);
+
+	FTexture2DRHIRef DestTexture = RectParams.SourceTexture;
+
+	TShaderMapRef<FBlitUIToHDRPS> PixelShader(ShaderMap);
+
+	RHICmdList.Transition(FRHITransitionInfo(DestTexture, ERHIAccess::Unknown, ERHIAccess::RTV));
+
+	const FVector2f InvSrcTextureSize(1.f / SrcTextureWidth, 1.f / SrcTextureHeight);
+
+	// Add some guard band to ensure blur will reach these pixels. It will overwrite pixels below, but these are going to be composited at the end anyway
+	const FVector2f UVStart = FVector2f(DestRect.Left - 10.0f, DestRect.Top - 10.0f) * InvSrcTextureSize;
+	const FVector2f UVEnd = FVector2f(DestRect.Right + 10.0f, DestRect.Bottom + 10.0f) * InvSrcTextureSize;
+	const FVector2f SizeUV = UVEnd - UVStart;
+
+	RHICmdList.SetViewport(0, 0, 0, SrcTextureWidth, SrcTextureHeight, 0.0f);
+	RHICmdList.SetScissorRect(false, 0, 0, 0, 0);
+
+	FRHIRenderPassInfo RPInfo(DestTexture, ERenderTargetActions::Load_Store);
+	RHICmdList.BeginRenderPass(RPInfo, TEXT("BlitUIToHDR"));
+	{
+		FGraphicsPipelineStateInitializer GraphicsPSOInit;
+		RHICmdList.ApplyCachedRenderTargets(GraphicsPSOInit);
+		GraphicsPSOInit.BlendState = TStaticBlendState<CW_RGB, BO_Add, BF_One, BF_InverseSourceAlpha, BO_Add, BF_One, BF_InverseSourceAlpha>::GetRHI();
+		GraphicsPSOInit.RasterizerState = TStaticRasterizerState<>::GetRHI();
+		GraphicsPSOInit.DepthStencilState = TStaticDepthStencilState<false, CF_Always>::GetRHI();
+
+		GraphicsPSOInit.BoundShaderState.VertexDeclarationRHI = GFilterVertexDeclaration.VertexDeclarationRHI;
+		GraphicsPSOInit.BoundShaderState.VertexShaderRHI = VertexShader.GetVertexShader();
+		GraphicsPSOInit.BoundShaderState.PixelShaderRHI = PixelShader.GetPixelShader();
+		GraphicsPSOInit.PrimitiveType = PT_TriangleList;
+		SetGraphicsPipelineState(RHICmdList, GraphicsPSOInit, 0);
+
+		FRHITexture* UITargetRTMaskTexture = RHISupportsRenderTargetWriteMask(GMaxRHIShaderPlatform) ? UITargetRTMask->GetRHI() : nullptr;
+
+		FMatrix44f Srgb_2_XYZ = GamutToXYZMatrix(EDisplayColorGamut::sRGB_D65);
+		FMatrix44f XYZ_2_Output = XYZToGamutMatrix(RectParams.HDRDisplayColorGamut);
+		// note: we use mul(m,v) instead of mul(v,m) in the shaders for color conversions which is why matrix multiplication is reversed compared to what we usually do
+		FMatrix44f CombinedMatrix = XYZ_2_Output * Srgb_2_XYZ;
+
+		PixelShader->SetParameters(RHICmdList, SourceTexture, UITargetRTMaskTexture, CombinedMatrix);
+
+		RendererModule.DrawRectangle(
+			RHICmdList,
+			UVStart.X * SrcTextureWidth, UVStart.Y * SrcTextureHeight,
+			SizeUV.X * SrcTextureWidth, SizeUV.Y * SrcTextureHeight,
+			UVStart.X, UVStart.Y,
+			SizeUV.X, SizeUV.Y,
+			FIntPoint(SrcTextureWidth, SrcTextureHeight),
+			FIntPoint(1, 1),
+			VertexShader,
+			EDRF_Default);
+	}
+	RHICmdList.EndRenderPass();
+}
+
 void FSlatePostProcessor::BlurRect(FRHICommandListImmediate& RHICmdList, IRendererModule& RendererModule, const FBlurRectParams& Params, const FPostProcessRectParams& RectParams)
 {
 	SCOPE_CYCLE_COUNTER(STAT_SlatePostProcessingRTTime);
@@ -51,7 +198,13 @@ void FSlatePostProcessor::BlurRect(FRHICommandListImmediate& RHICmdList, IRender
 
 	const FIntPoint DownsampleSize = RequiredSize;
 
-	IntermediateTargets->Update(RequiredSize);
+	IntermediateTargets->Update(RequiredSize, RectParams.SourceTexture);
+
+	if (RectParams.UITarget.IsValid() && RectParams.UITarget->GetRHI() != RectParams.SourceTexture.GetReference())
+	{
+		// in HDR mode, we are going to blur SourceTexture, but still need to take into account the UI already rendered. Blit UI into HDR target
+		BlitUIToHDRScene(RHICmdList, RendererModule, RectParams);
+	}
 
 	if(bDownsample)
 	{
@@ -207,7 +360,7 @@ void FSlatePostProcessor::ColorDeficiency(FRHICommandListImmediate& RHICmdList,
 	FIntPoint DestRectSize = RectParams.DestRect.GetSize().IntPoint();
 	FIntPoint RequiredSize = DestRectSize;
 
-	IntermediateTargets->Update(RequiredSize);
+	IntermediateTargets->Update(RequiredSize, RectParams.SourceTexture);
 
 #if 1
 	FSamplerStateRHIRef PointClamp = TStaticSamplerState<SF_Point, AM_Clamp, AM_Clamp, AM_Clamp>::GetRHI();
@@ -406,6 +559,25 @@ void FSlatePostProcessor::UpsampleRect(FRHICommandListImmediate& RHICmdList, IRe
 	RHICmdList.Transition(FRHITransitionInfo(DestTexture, ERHIAccess::Unknown, ERHIAccess::RTV));
 
 	FRHIRenderPassInfo RPInfo(DestTexture, ERenderTargetActions::Load_Store);
+
+	bool bHasMRT = false;
+
+	FRHITexture* UITargetTexture = Params.UITarget.IsValid() ? Params.UITarget->GetRHI() : nullptr;
+
+	if (UITargetTexture != nullptr && DestTexture != UITargetTexture)
+	{
+		RPInfo.ColorRenderTargets[1].RenderTarget = UITargetTexture;
+		RPInfo.ColorRenderTargets[1].ArraySlice = -1;
+		RPInfo.ColorRenderTargets[1].Action = ERenderTargetActions::Load_Store;
+		if (!RHISupportsRenderTargetWriteMask(GMaxRHIShaderPlatform))
+		{
+			RPInfo.ColorRenderTargets[2].RenderTarget = Params.UITargetMask->GetRHI();
+			RPInfo.ColorRenderTargets[2].ArraySlice = -1;
+			RPInfo.ColorRenderTargets[2].Action = ERenderTargetActions::Load_Store;
+		}
+		bHasMRT = true;
+	}
+
 	RHICmdList.BeginRenderPass(RPInfo, TEXT("UpsampleRect"));
 	{
 		RHICmdList.ApplyCachedRenderTargets(GraphicsPSOInit);
@@ -417,7 +589,15 @@ void FSlatePostProcessor::UpsampleRect(FRHICommandListImmediate& RHICmdList, IRe
 			Params.RestoreStateFunc(RHICmdList, GraphicsPSOInit);
 		}
 
-		TShaderMapRef<FSlatePostProcessUpsamplePS> PixelShader(ShaderMap);
+		TShaderRef<FSlateElementPS> PixelShader;
+		if (bHasMRT)
+		{
+			PixelShader = TShaderMapRef<FSlatePostProcessUpsamplePS<true> >(ShaderMap);
+		}
+		else
+		{
+			PixelShader = TShaderMapRef<FSlatePostProcessUpsamplePS<false> >(ShaderMap);
+		}
 
 		GraphicsPSOInit.BoundShaderState.VertexDeclarationRHI = GFilterVertexDeclaration.VertexDeclarationRHI;
 		GraphicsPSOInit.BoundShaderState.VertexShaderRHI = VertexShader.GetVertexShader();
diff --git a/Engine/Source/Runtime/SlateRHIRenderer/Private/SlatePostProcessor.h b/Engine/Source/Runtime/SlateRHIRenderer/Private/SlatePostProcessor.h
--- a/Engine/Source/Runtime/SlateRHIRenderer/Private/SlatePostProcessor.h
+++ b/Engine/Source/Runtime/SlateRHIRenderer/Private/SlatePostProcessor.h
@@ -8,6 +8,7 @@
 
 class FSlatePostProcessResource;
 class IRendererModule;
+struct IPooledRenderTarget;
 
 struct FPostProcessRectParams
 {
@@ -17,7 +18,10 @@ struct FPostProcessRectParams
 	FVector4f CornerRadius;
 	FIntPoint SourceTextureSize;
 	TFunction<void(FRHICommandListImmediate&, FGraphicsPipelineStateInitializer&)> RestoreStateFunc;
+	TRefCountPtr<IPooledRenderTarget> UITarget; // not using FTextureRHIRef because we want to be able to use FRenderTargetWriteMask::Decode
+	TRefCountPtr<IPooledRenderTarget> UITargetMask; // not using FTextureRHIRef because we want to be able to use FRenderTargetWriteMask::Decode
 	uint32 StencilRef{};
+	EDisplayColorGamut HDRDisplayColorGamut;
 };
 
 struct FBlurRectParams
diff --git a/Engine/Source/Runtime/SlateRHIRenderer/Private/SlateRHIRenderer.cpp b/Engine/Source/Runtime/SlateRHIRenderer/Private/SlateRHIRenderer.cpp
--- a/Engine/Source/Runtime/SlateRHIRenderer/Private/SlateRHIRenderer.cpp
+++ b/Engine/Source/Runtime/SlateRHIRenderer/Private/SlateRHIRenderer.cpp
@@ -663,7 +663,7 @@ public:
 		SetShaderValue(RHICmdList, RHICmdList.GetBoundPixelShader(), UILevel, CVarUILevel.GetValueOnRenderThread());
 		SetShaderValue(RHICmdList, RHICmdList.GetBoundPixelShader(), OutputDevice, CVarOutputDevice->GetValueOnRenderThread());
 		
-		if (RHISupportsRenderTargetWriteMask(GMaxRHIShaderPlatform))
+		if (UITextureWriteMaskRHI != nullptr)
 		{
 			SetTextureParameter(RHICmdList, RHICmdList.GetBoundPixelShader(), UIWriteMaskTexture, UITextureWriteMaskRHI);
 		}
@@ -697,6 +697,53 @@ private:
 	LAYOUT_FIELD(FShaderParameter, OutputDevice);
 };
 
+
+// Pixel shader to convert UI from PQ 2020 to Linear (and 2020 space)
+class FHDRBackBufferConvertPS : public FGlobalShader
+{
+	DECLARE_SHADER_TYPE(FHDRBackBufferConvertPS, Global);
+public:
+
+	static bool ShouldCompilePermutation(const FGlobalShaderPermutationParameters& Parameters)
+	{
+		return IsFeatureLevelSupported(Parameters.Platform, ERHIFeatureLevel::SM5);
+	}
+
+	FHDRBackBufferConvertPS(const ShaderMetaType::CompiledShaderInitializerType& Initializer) :
+		FGlobalShader(Initializer)
+	{
+		SceneTexture.Bind(Initializer.ParameterMap, TEXT("SceneTexture"));
+		SceneSampler.Bind(Initializer.ParameterMap, TEXT("SceneSampler"));
+	}
+	FHDRBackBufferConvertPS() {}
+
+	void SetParameters(FRHICommandList& RHICmdList, FRHITexture* SceneTextureRHI)
+	{
+		SetTextureParameter(RHICmdList, RHICmdList.GetBoundPixelShader(), SceneTexture, SceneSampler, TStaticSamplerState<SF_Point>::GetRHI(), SceneTextureRHI);
+	}
+
+	static void ModifyCompilationEnvironment(const FGlobalShaderPermutationParameters& Parameters, FShaderCompilerEnvironment& OutEnvironment)
+	{
+		FGlobalShader::ModifyCompilationEnvironment(Parameters, OutEnvironment);
+	}
+
+	static const TCHAR* GetSourceFilename()
+	{
+		return TEXT("/Engine/Private/CompositeUIPixelShader.usf");
+	}
+
+	static const TCHAR* GetFunctionName()
+	{
+		return TEXT("HDRBackBufferConvertPS");
+	}
+
+private:
+	LAYOUT_FIELD(FShaderResourceParameter, SceneTexture);
+	LAYOUT_FIELD(FShaderResourceParameter, SceneSampler);
+};
+
+IMPLEMENT_SHADER_TYPE(, FHDRBackBufferConvertPS, FHDRBackBufferConvertPS::GetSourceFilename(), FHDRBackBufferConvertPS::GetFunctionName(), SF_Pixel);
+
 #define SHADER_VARIATION(A) typedef FCompositePS<A> FCompositePS##A; \
 	IMPLEMENT_SHADER_TYPE2(FCompositePS##A, SF_Pixel);
 SHADER_VARIATION(0)  SHADER_VARIATION(1)
@@ -755,6 +802,12 @@ void RenderSlateBatch(FTexture2DRHIRef SlateRenderTarget, bool bClear, bool bIsH
 				FSlateRenderingParams RenderParams(ViewMatrix * ViewportInfo.ProjectionMatrix, DrawCommandParams.Time);
 				RenderParams.bWireFrame = !!SlateWireFrame;
 				RenderParams.bIsHDR = bIsHDR;
+				RenderParams.HDRDisplayColorGamut = ViewportInfo.HDRDisplayColorGamut;
+				if (ViewportInfo.bSceneHDREnabled && !bIsHDR)
+				{
+					RenderParams.UITarget = ViewportInfo.UITargetRT;
+					RenderParams.UITargetMask = ViewportInfo.UITargetRTMask;
+				}
 				RenderingPolicy->SetUseGammaCorrection(!bIsHDR);
 
 				FTexture2DRHIRef EmptyTarget;
@@ -851,14 +904,6 @@ void FSlateRHIRenderer::DrawWindow_RenderThread(FRHICommandListImmediate& RHICmd
 			SCOPE_CYCLE_COUNTER(STAT_SlateRenderingRTTime);
 			CSV_SCOPED_TIMING_STAT_EXCLUSIVE(Slate);
 
-			FSlateBatchData& BatchData = WindowElementList.GetBatchData();
-
-			// Update the vertex and index buffer	
-			RenderingPolicy->BuildRenderingBuffers(RHICmdList, BatchData);
-
-			// This must happen after rendering buffers are created
-			ViewportInfo.ConditionallyUpdateDepthBuffer(BatchData.IsStencilClippingRequired(), ViewportInfo.DesiredWidth, ViewportInfo.DesiredHeight);
-
 			// should have been created by the game thread
 			check(IsValidRef(ViewportInfo.ViewportRHI));
 
@@ -869,27 +914,56 @@ void FSlateRHIRenderer::DrawWindow_RenderThread(FRHICommandListImmediate& RHICmd
 			const uint32 ViewportWidth = (ViewportRT) ? ViewportRT->GetSizeX() : ViewportInfo.Width;
 			const uint32 ViewportHeight = (ViewportRT) ? ViewportRT->GetSizeY() : ViewportInfo.Height;
 
+			EPixelFormat HDRSourceRTExpectedFormat = BackBuffer->GetFormat();
+			if (bCompositeUI)
+			{
+				// Some effects like slate posprocess blur get quite expensive if we have to do PQ->linear for every texture sample. Instead, work with a format where values are stored linearly
+				// Do not test BackBuffer->GetFormat() == PF_A2B10G10R10 because some platforms do not return this PF with ST2084
+				if (ViewportInfo.HDRDisplayOutputFormat == EDisplayOutputFormat::HDR_ACES_1000nit_ST2084 || ViewportInfo.HDRDisplayOutputFormat == EDisplayOutputFormat::HDR_ACES_2000nit_ST2084)
+				{
+					// Prefer small float format to improve performance for slate's background blur. Banding might appear when choosing this format, but UITargetRTMask allows to discard pixels that have
+                    // not been touched by the UI
+					if (GPixelFormats[PF_FloatR11G11B10].Supported)
+					{
+						HDRSourceRTExpectedFormat = PF_FloatR11G11B10;
+					}
+					else
+					{
+						HDRSourceRTExpectedFormat = PF_FloatRGBA;
+					}
+				}
+			}
+
 			// Check to see that targets are up-to-date
 			if (bCompositeUI && (!ViewportInfo.UITargetRT || ViewportInfo.UITargetRT->GetRHI()->GetSizeX() != ViewportWidth || ViewportInfo.UITargetRT->GetRHI()->GetSizeY() != ViewportHeight
-				|| !ViewportInfo.HDRSourceRT || ViewportInfo.HDRSourceRT->GetRHI()->GetFormat() != BackBuffer->GetFormat()))
+				|| !ViewportInfo.HDRSourceRT || ViewportInfo.HDRSourceRT->GetRHI()->GetFormat() != HDRSourceRTExpectedFormat))
 			{
 				// Composition buffers
 				{
-				ETextureCreateFlags BaseFlags = RHISupportsRenderTargetWriteMask(GMaxRHIShaderPlatform) ? TexCreate_NoFastClearFinalize | TexCreate_DisableDCC : TexCreate_None;
-				FPooledRenderTargetDesc Desc(FPooledRenderTargetDesc::Create2DDesc(FIntPoint(ViewportWidth, ViewportHeight),
-					GetSlateRecommendedColorFormat(),
-					FClearValueBinding::Transparent,
-					BaseFlags,
-					TexCreate_ShaderResource | TexCreate_RenderTargetable,
-					false,
-					1,
-					true,
-					true));
-
-				GRenderTargetPool.FindFreeElement(RHICmdList, Desc, ViewportInfo.UITargetRT, TEXT("UITargetRT"));
-
-				Desc.Format = BackBuffer->GetFormat();
-				GRenderTargetPool.FindFreeElement(RHICmdList, Desc, ViewportInfo.HDRSourceRT, TEXT("HDRSourceRT"));
+					ETextureCreateFlags BaseFlags = RHISupportsRenderTargetWriteMask(GMaxRHIShaderPlatform) ? TexCreate_NoFastClearFinalize | TexCreate_DisableDCC : TexCreate_None;
+					FPooledRenderTargetDesc Desc(FPooledRenderTargetDesc::Create2DDesc(FIntPoint(ViewportWidth, ViewportHeight),
+						GetSlateRecommendedColorFormat(),
+						FClearValueBinding::Transparent,
+						BaseFlags,
+						TexCreate_ShaderResource | TexCreate_RenderTargetable,
+						false,
+						1,
+						true,
+						true));
+
+					GRenderTargetPool.FindFreeElement(RHICmdList, Desc, ViewportInfo.UITargetRT, TEXT("UITargetRT"));
+
+					Desc.Format = HDRSourceRTExpectedFormat;
+					GRenderTargetPool.FindFreeElement(RHICmdList, Desc, ViewportInfo.HDRSourceRT, TEXT("HDRSourceRT"));
+
+                    // Because we may use small float format for storing scene color data, keep a mask to discard pixels that were untouched by the UI. Note: relying only on the UI Rendertarget's alpha does 
+                    // not work because of FSlatePostProcessor::BlurRect. It composes UI to scene / blur it / remove the UI that was there before the blur. After that, UI pixels are merged with the scene pixels
+                    // This RT will keep track of where the blur happened, and will be used with UITargetRT's alpha channel
+					if (!RHISupportsRenderTargetWriteMask(GMaxRHIShaderPlatform))
+					{
+						Desc.Format = PF_R8_UINT;
+						GRenderTargetPool.FindFreeElement(RHICmdList, Desc, ViewportInfo.UITargetRTMask, TEXT("UITargetRTMask"));
+					}
 				}
 
 				// LUT
@@ -927,13 +1001,7 @@ void FSlateRHIRenderer::DrawWindow_RenderThread(FRHICommandListImmediate& RHICmd
 					RHICmdList.BeginRenderPass(RPInfo, TEXT("Clear back buffer"));
 					RHICmdList.EndRenderPass();
 				}
-				else
 #endif
-				{
-				// Grab HDR backbuffer
-					TransitionAndCopyTexture(RHICmdList, FinalBuffer, ViewportInfo.HDRSourceRT->GetRHI(), {});
-				}
-
 				// UI backbuffer is temp target
 				BackBuffer = ViewportInfo.UITargetRT->GetRHI();
 			}
@@ -967,20 +1035,96 @@ void FSlateRHIRenderer::DrawWindow_RenderThread(FRHICommandListImmediate& RHICmd
 			RHICmdList.BeginDrawingViewport(ViewportInfo.ViewportRHI, FTextureRHIRef());
 			RHICmdList.SetViewport(0.f, 0.f, 0.f, (float)ViewportWidth, (float)ViewportHeight, 0.0f);
 
-			bool bHdrTarget = ViewportInfo.bSceneHDREnabled && !bCompositeUI;
-			RenderSlateBatch(BackBuffer, bClear, bHdrTarget, ViewportInfo, ViewMatrix, BatchData, RHICmdList, ViewportWidth, ViewportHeight, DrawCommandParams, RenderingPolicy, PostProcessBuffer);
-
 			if (bCompositeUI)
-					{
+			{
 				FSlateBatchData& BatchDataHDR = WindowElementList.GetBatchDataHDR();
 				bool bHasBatches = BatchDataHDR.GetRenderBatches().Num() > 0;
-						if (bHasBatches)
-						{
+				if (bHasBatches)
+				{
 					RenderingPolicy->BuildRenderingBuffers(RHICmdList, BatchDataHDR);
-					FTexture2DRHIRef UITargetHDRRTRHI(ViewportInfo.HDRSourceRT->GetRHI());
+					FTexture2DRHIRef UITargetHDRRTRHI(FinalBuffer);
 					RenderSlateBatch(UITargetHDRRTRHI, /*bClear*/ false, /*bIsHDR*/ true, ViewportInfo, ViewMatrix, BatchDataHDR, RHICmdList, ViewportWidth, ViewportHeight, DrawCommandParams, RenderingPolicy, PostProcessBuffer);
-			}
+			    }
 
+				// Grab HDR backbuffer
+				if (ViewportInfo.HDRSourceRT->GetDesc().Format == FinalBuffer->GetDesc().Format)
+				{
+					TransitionAndCopyTexture(RHICmdList, FinalBuffer, ViewportInfo.HDRSourceRT->GetRHI(), {});
+				}
+				else
+				{
+					// Do not test FinalBuffer->GetDesc().Format == PF_A2B10G10R10 because some platforms do not return this PF even with ST2084
+					//ensure(FinalBuffer->GetDesc().Format == PF_A2B10G10R10);
+					ensure(ViewportInfo.HDRDisplayOutputFormat == EDisplayOutputFormat::HDR_ACES_1000nit_ST2084 || ViewportInfo.HDRDisplayOutputFormat == EDisplayOutputFormat::HDR_ACES_2000nit_ST2084);
+					ensure(ViewportInfo.HDRSourceRT->GetDesc().Format == PF_FloatR11G11B10 || ViewportInfo.HDRSourceRT->GetDesc().Format == PF_FloatRGBA);
+
+					// Convert Scene Color in PQ to Linear so that Slate can perform post processing on it without any conversion
+					{
+						auto ShaderMap = GetGlobalShaderMap(GMaxRHIFeatureLevel);
+						RHICmdList.Transition({ FRHITransitionInfo(ViewportInfo.HDRSourceRT->GetRHI(), ERHIAccess::Unknown, ERHIAccess::RTV),
+							                    FRHITransitionInfo(FinalBuffer, ERHIAccess::Unknown, ERHIAccess::SRVGraphics)
+							                  });
+						FRHIRenderPassInfo RPInfo(ViewportInfo.HDRSourceRT->GetRHI(), ERenderTargetActions::Load_Store);
+						RHICmdList.BeginRenderPass(RPInfo, TEXT("ConvertBackBuffer"));
+						{
+							FGraphicsPipelineStateInitializer GraphicsPSOInit;
+							RHICmdList.ApplyCachedRenderTargets(GraphicsPSOInit);
+							GraphicsPSOInit.BlendState = TStaticBlendState<>::GetRHI();
+							GraphicsPSOInit.RasterizerState = TStaticRasterizerState<>::GetRHI();
+							GraphicsPSOInit.DepthStencilState = TStaticDepthStencilState<false, CF_Always>::GetRHI();
+
+							TShaderMapRef<FScreenVS> VertexShader(ShaderMap);
+							TShaderMapRef<FHDRBackBufferConvertPS> PixelShader(ShaderMap);
+
+							GraphicsPSOInit.BoundShaderState.VertexDeclarationRHI = GFilterVertexDeclaration.VertexDeclarationRHI;
+							GraphicsPSOInit.BoundShaderState.VertexShaderRHI = VertexShader.GetVertexShader();
+							GraphicsPSOInit.BoundShaderState.PixelShaderRHI = PixelShader.GetPixelShader();
+							GraphicsPSOInit.PrimitiveType = PT_TriangleList;
+
+							SetGraphicsPipelineState(RHICmdList, GraphicsPSOInit, 0);
+
+							PixelShader->SetParameters(RHICmdList, FinalBuffer);
+
+							static const FName RendererModuleName("Renderer");
+							IRendererModule& RendererModule = FModuleManager::GetModuleChecked<IRendererModule>(RendererModuleName);
+							RendererModule.DrawRectangle(
+								RHICmdList,
+								0.f, 0.f,
+								(float)ViewportWidth, (float)ViewportHeight,
+								0.f, 0.f,
+								(float)ViewportWidth, (float)ViewportHeight,
+								FIntPoint(ViewportWidth, ViewportHeight),
+								FIntPoint(ViewportWidth, ViewportHeight),
+								VertexShader,
+								EDRF_UseTriangleOptimization);
+						}
+						RHICmdList.EndRenderPass();
+					}
+				}
+
+				if (!RHISupportsRenderTargetWriteMask(GMaxRHIShaderPlatform) && ViewportInfo.UITargetRTMask.IsValid())
+				{
+					FRHIRenderPassInfo RPInfo(ViewportInfo.UITargetRTMask->GetRHI(), ERenderTargetActions::Clear_Store);
+					RHICmdList.BeginRenderPass(RPInfo, TEXT("Clear_UITargetRTMask"));
+					RHICmdList.EndRenderPass();
+				}
+
+		    	PostProcessBuffer = ViewportInfo.HDRSourceRT->GetRHI();
+            }
+
+			FSlateBatchData& BatchData = WindowElementList.GetBatchData();
+
+			// Update the vertex and index buffer	
+		    RenderingPolicy->BuildRenderingBuffers(RHICmdList, BatchData);
+
+			// This must happen after rendering buffers are created
+			ViewportInfo.ConditionallyUpdateDepthBuffer(BatchData.IsStencilClippingRequired(), ViewportInfo.DesiredWidth, ViewportInfo.DesiredHeight);
+
+		    bool bHdrTarget = ViewportInfo.bSceneHDREnabled && !bCompositeUI;
+		    RenderSlateBatch(BackBuffer, bClear, bHdrTarget, ViewportInfo, ViewMatrix, BatchData, RHICmdList, ViewportWidth, ViewportHeight, DrawCommandParams, RenderingPolicy, PostProcessBuffer);
+
+			if (bCompositeUI)
+			{
 				SCOPED_DRAW_EVENT(RHICmdList, SlateUI_Composition);
 
 				static const FName RendererModuleName("Renderer");
@@ -1055,7 +1199,7 @@ void FSlateRHIRenderer::DrawWindow_RenderThread(FRHICommandListImmediate& RHICmd
 
 						TShaderMapRef<FScreenVS> VertexShader(ShaderMap);
 
-						FRHITexture* UITargetRTMaskTexture = RHISupportsRenderTargetWriteMask(GMaxRHIShaderPlatform) ? ViewportInfo.UITargetRTMask->GetRHI() : nullptr;
+						FRHITexture* UITargetRTMaskTexture = ViewportInfo.UITargetRTMask.IsValid() ? ViewportInfo.UITargetRTMask->GetRHI() : nullptr;
 						if (HDROutputDevice == (int32)EDisplayOutputFormat::HDR_ACES_1000nit_ScRGB || HDROutputDevice == (int32)EDisplayOutputFormat::HDR_ACES_2000nit_ScRGB)
 						{
 							// ScRGB encoding
diff --git a/Engine/Source/Runtime/SlateRHIRenderer/Private/SlateRHIRenderingPolicy.cpp b/Engine/Source/Runtime/SlateRHIRenderer/Private/SlateRHIRenderingPolicy.cpp
--- a/Engine/Source/Runtime/SlateRHIRenderer/Private/SlateRHIRenderingPolicy.cpp
+++ b/Engine/Source/Runtime/SlateRHIRenderer/Private/SlateRHIRenderingPolicy.cpp
@@ -1163,6 +1163,9 @@ void FSlateRHIRenderingPolicy::DrawElements(
 				RectParams.DestRect = FSlateRect(QuadPositionData.X, QuadPositionData.Y, QuadPositionData.Z, QuadPositionData.W);
 				RectParams.SourceTextureSize = PostProcessTexture->GetSizeXY();
 				RectParams.CornerRadius = ShaderParams.PixelParams3;
+				RectParams.UITarget = Params.UITarget;
+				RectParams.UITargetMask = Params.UITargetMask;
+				RectParams.HDRDisplayColorGamut = Params.HDRDisplayColorGamut;
 
 				RectParams.RestoreStateFunc = [&](FRHICommandListImmediate&InRHICmdList, FGraphicsPipelineStateInitializer& InGraphicsPSOInit) {
 					return UpdateScissorRect(
@@ -1171,7 +1174,7 @@ void FSlateRHIRenderingPolicy::DrawElements(
 						ScissorClips,
 						StencilClips,
 #endif
-						StencilRef,
+						RectParams.StencilRef,
 						MaskingID,
 						BackBuffer,
 						RenderBatch,
diff --git a/Engine/Source/Runtime/SlateRHIRenderer/Private/SlateRHIRenderingPolicy.h b/Engine/Source/Runtime/SlateRHIRenderer/Private/SlateRHIRenderingPolicy.h
--- a/Engine/Source/Runtime/SlateRHIRenderer/Private/SlateRHIRenderingPolicy.h
+++ b/Engine/Source/Runtime/SlateRHIRenderer/Private/SlateRHIRenderingPolicy.h
@@ -29,6 +29,9 @@ struct FSlateRenderingParams
 	FMatrix44f ViewProjectionMatrix;
 	FVector2f ViewOffset;
 	FGameTime Time;
+	TRefCountPtr<IPooledRenderTarget> UITarget;
+	TRefCountPtr<IPooledRenderTarget> UITargetMask;
+	EDisplayColorGamut HDRDisplayColorGamut;
 	bool bWireFrame;
 	bool bIsHDR;
 
@@ -36,6 +39,7 @@ struct FSlateRenderingParams
 		: ViewProjectionMatrix(InViewProjectionMatrix)
 		, ViewOffset(0.f, 0.f)
 		, Time(InTime)
+		, HDRDisplayColorGamut(EDisplayColorGamut::sRGB_D65)
 		, bWireFrame(false)
 		, bIsHDR(false)
 	{
diff --git a/Engine/Source/Runtime/SlateRHIRenderer/Private/SlateShaders.cpp b/Engine/Source/Runtime/SlateRHIRenderer/Private/SlateShaders.cpp
--- a/Engine/Source/Runtime/SlateRHIRenderer/Private/SlateShaders.cpp
+++ b/Engine/Source/Runtime/SlateRHIRenderer/Private/SlateShaders.cpp
@@ -18,7 +18,8 @@ IMPLEMENT_SHADER_TYPE(, FSlateDebugOverdrawPS, TEXT("/Engine/Private/SlateElemen
 
 IMPLEMENT_SHADER_TYPE(, FSlatePostProcessBlurPS, TEXT("/Engine/Private/SlatePostProcessPixelShader.usf"), TEXT("GaussianBlurMain"), SF_Pixel);
 IMPLEMENT_SHADER_TYPE(, FSlatePostProcessDownsamplePS, TEXT("/Engine/Private/SlatePostProcessPixelShader.usf"), TEXT("DownsampleMain"), SF_Pixel);
-IMPLEMENT_SHADER_TYPE(, FSlatePostProcessUpsamplePS, TEXT("/Engine/Private/SlatePostProcessPixelShader.usf"), TEXT("UpsampleMain"), SF_Pixel);
+IMPLEMENT_SHADER_TYPE(template<>, FSlatePostProcessUpsamplePS<true>, TEXT("/Engine/Private/SlatePostProcessPixelShader.usf"), TEXT("UpsampleMain"), SF_Pixel);
+IMPLEMENT_SHADER_TYPE(template<>, FSlatePostProcessUpsamplePS<false>, TEXT("/Engine/Private/SlatePostProcessPixelShader.usf"), TEXT("UpsampleMain"), SF_Pixel);
 IMPLEMENT_SHADER_TYPE(, FSlatePostProcessColorDeficiencyPS, TEXT("/Engine/Private/SlatePostProcessColorDeficiencyPixelShader.usf"), TEXT("ColorDeficiencyMain"), SF_Pixel);
 
 IMPLEMENT_SHADER_TYPE(, FSlateMaskingVS, TEXT("/Engine/Private/SlateMaskingShader.usf"), TEXT("MainVS"), SF_Vertex);
diff --git a/Engine/Source/Runtime/SlateRHIRenderer/Private/SlateShaders.h b/Engine/Source/Runtime/SlateRHIRenderer/Private/SlateShaders.h
--- a/Engine/Source/Runtime/SlateRHIRenderer/Private/SlateShaders.h
+++ b/Engine/Source/Runtime/SlateRHIRenderer/Private/SlateShaders.h
@@ -394,7 +394,7 @@ private:
 	LAYOUT_FIELD(FShaderParameter, UVBounds);
 };
 
-
+template<bool bOutputToUITarget>
 class FSlatePostProcessUpsamplePS : public FSlateElementPS
 {
 	DECLARE_SHADER_TYPE(FSlatePostProcessUpsamplePS, Global);
@@ -415,6 +415,16 @@ public:
 	{
 	}
 
+	/**
+	 * Modifies the compilation of this shader
+	 */
+	static void ModifyCompilationEnvironment(const FGlobalShaderPermutationParameters& Parameters, FShaderCompilerEnvironment& OutEnvironment)
+	{
+		// Set defines based on what this shader will be used for
+		OutEnvironment.SetDefine(TEXT("OUTPUT_TO_UI_TARGET"), (uint32)(bOutputToUITarget ? 1 : 0));
+		FSlateElementPS::ModifyCompilationEnvironment(Parameters, OutEnvironment);
+	}
+
 private:
 };
 
