Coding Standard fix to change "master" prefix to "source" for index and Vertex buffers

diff --git a/Engine/Source/Runtime/SlateRHIRenderer/Private/SlateRHIRenderingPolicy.cpp b/Engine/Source/Runtime/SlateRHIRenderer/Private/SlateRHIRenderingPolicy.cpp
--- a/Engine/Source/Runtime/SlateRHIRenderer/Private/SlateRHIRenderingPolicy.cpp
+++ b/Engine/Source/Runtime/SlateRHIRenderer/Private/SlateRHIRenderingPolicy.cpp
@@ -104,16 +104,16 @@ void FSlateRHIRenderingPolicy::InitResources()
 
 	UE_LOG(LogSlate, Verbose, TEXT("Allocating space for %d vertices"), NumVertices);
 
-	MasterVertexBuffer.Init(NumVertices);
-	MasterIndexBuffer.Init(NumVertices);
+	SourceVertexBuffer.Init(NumVertices);
+	SourceIndexBuffer.Init(NumVertices);
 
 	BeginInitResource(&StencilVertexBuffer);
 }
 
 void FSlateRHIRenderingPolicy::ReleaseResources()
 {
-	MasterVertexBuffer.Destroy();
-	MasterIndexBuffer.Destroy();
+	SourceVertexBuffer.Destroy();
+	SourceIndexBuffer.Destroy();
 
 	BeginReleaseResource(&StencilVertexBuffer);
 }
@@ -149,12 +149,12 @@ void FSlateRHIRenderingPolicy::BuildRenderingBuffers(FRHICommandListImmediate& R
 		bool bShouldShrinkResources = false;
 		bool bAbsoluteIndices = CVarSlateAbsoluteIndices.GetValueOnRenderThread() != 0;
 
-		MasterVertexBuffer.PreFillBuffer(NumVertices, bShouldShrinkResources);
-		MasterIndexBuffer.PreFillBuffer(NumIndices, bShouldShrinkResources);
+		SourceVertexBuffer.PreFillBuffer(NumVertices, bShouldShrinkResources);
+		SourceIndexBuffer.PreFillBuffer(NumIndices, bShouldShrinkResources);
 
 		RHICmdList.EnqueueLambda([
-			VertexBuffer = MasterVertexBuffer.VertexBufferRHI.GetReference(),
-			IndexBuffer = MasterIndexBuffer.IndexBufferRHI.GetReference(),
+			VertexBuffer = SourceVertexBuffer.VertexBufferRHI.GetReference(),
+			IndexBuffer = SourceIndexBuffer.IndexBufferRHI.GetReference(),
 			&InBatchData,
 			bAbsoluteIndices
 		](FRHICommandListImmediate& InRHICmdList)
@@ -184,8 +184,8 @@ void FSlateRHIRenderingPolicy::BuildRenderingBuffers(FRHICommandListImmediate& R
 		RHICmdList.RHIThreadFence(true);
 	}
 
-	checkSlow(MasterVertexBuffer.GetBufferUsageSize() <= MasterVertexBuffer.GetBufferSize());
-	checkSlow(MasterIndexBuffer.GetBufferUsageSize() <= MasterIndexBuffer.GetBufferSize());
+	checkSlow(SourceVertexBuffer.GetBufferUsageSize() <= SourceVertexBuffer.GetBufferSize());
+	checkSlow(SourceIndexBuffer.GetBufferUsageSize() <= SourceIndexBuffer.GetBufferSize());
 
 	SET_DWORD_STAT(STAT_SlateNumLayers, InBatchData.GetNumLayers());
 	SET_DWORD_STAT(STAT_SlateNumBatches, InBatchData.GetNumFinalBatches());
@@ -626,8 +626,8 @@ void FSlateRHIRenderingPolicy::DrawElements(
 
 	FSamplerStateRHIRef BilinearClamp = TStaticSamplerState<SF_Bilinear, AM_Clamp, AM_Clamp, AM_Clamp>::GetRHI();
 
-	TSlateElementVertexBuffer<FSlateVertex>* VertexBufferPtr = &MasterVertexBuffer;
-	FSlateElementIndexBuffer* IndexBufferPtr = &MasterIndexBuffer;
+	TSlateElementVertexBuffer<FSlateVertex>* VertexBufferPtr = &SourceVertexBuffer;
+	FSlateElementIndexBuffer* IndexBufferPtr = &SourceIndexBuffer;
 
 	FGraphicsPipelineStateInitializer GraphicsPSOInit;
 	RHICmdList.ApplyCachedRenderTargets(GraphicsPSOInit);
@@ -672,8 +672,8 @@ void FSlateRHIRenderingPolicy::DrawElements(
 	*/
 	while (NextRenderBatchIndex != INDEX_NONE)
 	{
-		VertexBufferPtr = &MasterVertexBuffer;
-		IndexBufferPtr = &MasterIndexBuffer;
+		VertexBufferPtr = &SourceVertexBuffer;
+		IndexBufferPtr = &SourceIndexBuffer;
 
 		if (!RHICmdList.IsInsideRenderPass())
 		{
diff --git a/Engine/Source/Runtime/SlateRHIRenderer/Private/SlateRHIRenderingPolicy.h b/Engine/Source/Runtime/SlateRHIRenderer/Private/SlateRHIRenderingPolicy.h
--- a/Engine/Source/Runtime/SlateRHIRenderer/Private/SlateRHIRenderingPolicy.h
+++ b/Engine/Source/Runtime/SlateRHIRenderer/Private/SlateRHIRenderingPolicy.h
@@ -91,8 +91,8 @@ private:
 
 private:
 	/** Buffers used for rendering */
-	TSlateElementVertexBuffer<FSlateVertex> MasterVertexBuffer;
-	FSlateElementIndexBuffer MasterIndexBuffer;
+	TSlateElementVertexBuffer<FSlateVertex> SourceVertexBuffer;
+	FSlateElementIndexBuffer SourceIndexBuffer;
 
 	FSlateStencilClipVertexBuffer StencilVertexBuffer;
 
