Make a new RenderthreadCriticalPath stat which includes dependent waits. The new stat still excludes GPU occlusion query waits, the present wait and non-critical path waits (waiting for the game thread or RHI thread). Waits for tasks kicked on the GT are included though (unlike the existing stat).

The implementation extends the existing FThreadIdleStats to track critical path vs non-critical path waits for a thread.

The cvar r.RenderThreadTimeIncludesDependentWaits forces the main renderthread stat to use the new behavior, but this is disabled by default for historical tracking reasons.

diff --git a/Engine/Source/Runtime/Core/Private/ProfilingDebugging/CsvProfiler.cpp b/Engine/Source/Runtime/Core/Private/ProfilingDebugging/CsvProfiler.cpp
--- a/Engine/Source/Runtime/Core/Private/ProfilingDebugging/CsvProfiler.cpp
+++ b/Engine/Source/Runtime/Core/Private/ProfilingDebugging/CsvProfiler.cpp
@@ -25,6 +25,7 @@
 #include "HAL/Runnable.h"
 #include "Misc/EngineVersion.h"
 #include "Stats/Stats.h"
+#include "Stats/Stats2.h"
 #include "HAL/LowLevelMemTracker.h"
 #include "Misc/Compression.h"
 #include "Misc/ConfigCacheIni.h"
@@ -3362,7 +3363,14 @@ void FCsvProfiler::BeginWait()
 #if CSV_PROFILER_SUPPORT_NAMED_EVENTS
 			if (UNLIKELY(GCsvProfilerNamedEventsExclusive))
 			{
-				CSV_PROFILER_BeginNamedEvent(FColor(255, 128, 128), "CsvEventWait");
+				if ( FThreadIdleStats::Get().IsCriticalPath() )
+				{
+					CSV_PROFILER_BeginNamedEvent(FColor(192, 96, 96), "CsvEventWait");
+				}
+				else
+				{
+					CSV_PROFILER_BeginNamedEvent(FColor(255, 128, 128), "CsvEventWait (Non-CP)");
+				}
 			}
 #endif
 			FCsvProfilerThreadData::Get().AddTimestampExclusiveBegin(WaitStatName);
diff --git a/Engine/Source/Runtime/Core/Public/Stats/Stats2.h b/Engine/Source/Runtime/Core/Public/Stats/Stats2.h
--- a/Engine/Source/Runtime/Core/Public/Stats/Stats2.h
+++ b/Engine/Source/Runtime/Core/Public/Stats/Stats2.h
@@ -55,6 +55,8 @@ class CORE_API FThreadIdleStats : public TThreadSingleton<FThreadIdleStats>
 
 	FThreadIdleStats()
 		: Waits(0)
+		, WaitsCriticalPath(0)
+		, IsCriticalPathCounter(1)
 	{}
 
 public:
@@ -62,6 +64,47 @@ public:
 	/** Total cycles we waited for sleep or event. **/
 	uint32 Waits;
 
+	/** Total cycles we waited for sleep or event on the critical path. **/
+	uint32 WaitsCriticalPath;
+
+	int IsCriticalPathCounter;
+
+	static void BeginCriticalPath()
+	{
+		FThreadIdleStats::Get().IsCriticalPathCounter++;
+	}
+
+	static void EndCriticalPath()
+	{
+		FThreadIdleStats::Get().IsCriticalPathCounter--;
+	}
+
+	struct FScopeNonCriticalPath
+	{
+		FScopeNonCriticalPath()
+		{
+			FThreadIdleStats::Get().IsCriticalPathCounter--;
+		}
+		~FScopeNonCriticalPath()
+		{
+			FThreadIdleStats::Get().IsCriticalPathCounter++;
+		}
+	};
+
+	bool IsCriticalPath() const 
+	{
+		return IsCriticalPathCounter > 0;
+	}
+
+	void Reset()
+	{
+		Waits = 0;
+		WaitsCriticalPath = 0;
+		IsCriticalPathCounter = 1;
+	}
+
+
+
 	struct FScopeIdle
 	{
 		/** Starting cycle counter. */
@@ -87,7 +130,14 @@ public:
 		{
 			if( !bIgnore )
 			{
-				FThreadIdleStats::Get().Waits += FPlatformTime::Cycles() - Start;
+				FThreadIdleStats& IdleStats = FThreadIdleStats::Get();
+				uint32 CyclesElapsed = FPlatformTime::Cycles() - Start;
+				IdleStats.Waits += CyclesElapsed;
+
+				if (IdleStats.IsCriticalPath())
+				{
+					IdleStats.WaitsCriticalPath += CyclesElapsed;
+				}
 			}
 		}
 #endif
diff --git a/Engine/Source/Runtime/Engine/Private/UnrealClient.cpp b/Engine/Source/Runtime/Engine/Private/UnrealClient.cpp
--- a/Engine/Source/Runtime/Engine/Private/UnrealClient.cpp
+++ b/Engine/Source/Runtime/Engine/Private/UnrealClient.cpp
@@ -1642,7 +1642,7 @@ void FViewport::Draw( bool bShouldPresent /*= true */)
 
 						LastFrameUpdated = GFrameCounter;
 						Lastimestamp		= CurrentTime;
-						GameThread.Waits = 0;
+						GameThread.Reset();
 					}
 				}
 
diff --git a/Engine/Source/Runtime/Launch/Private/LaunchEngineLoop.cpp b/Engine/Source/Runtime/Launch/Private/LaunchEngineLoop.cpp
--- a/Engine/Source/Runtime/Launch/Private/LaunchEngineLoop.cpp
+++ b/Engine/Source/Runtime/Launch/Private/LaunchEngineLoop.cpp
@@ -1429,7 +1429,8 @@ static void UpdateCoreCsvStats_EndFrame()
 	    CSV_CUSTOM_STAT_GLOBAL(RenderThreadTime, FPlatformTime::ToMilliseconds(GRenderThreadTime), ECsvCustomStatOp::Set);
 	    CSV_CUSTOM_STAT_GLOBAL(GameThreadTime, FPlatformTime::ToMilliseconds(GGameThreadTime), ECsvCustomStatOp::Set);
 	    CSV_CUSTOM_STAT_GLOBAL(GPUTime, FPlatformTime::ToMilliseconds(GGPUFrameTime), ECsvCustomStatOp::Set);
-	    if (IsRunningRHIInSeparateThread())
+		CSV_CUSTOM_STAT_GLOBAL(RenderThreadTime_CriticalPath, FPlatformTime::ToMilliseconds(GRenderThreadTimeCriticalPath), ECsvCustomStatOp::Set);
+		if (IsRunningRHIInSeparateThread())
 	    {
 		    CSV_CUSTOM_STAT_GLOBAL(RHIThreadTime, FPlatformTime::ToMilliseconds(GRHIThreadTime), ECsvCustomStatOp::Set);
 	    }
@@ -5012,6 +5013,9 @@ static inline void BeginFrameRenderThread(FRHICommandListImmediate& RHICmdList,
 #if CSV_PROFILER
 	FCsvProfiler::BeginExclusiveStat("RenderThreadOther");
 #endif
+
+	// Waits after this point but before BeginRenderingViewFamilies are not on the RT critical path (since we're waiting for the GT)
+	FThreadIdleStats::EndCriticalPath();
 }
 
 
diff --git a/Engine/Source/Runtime/RHI/Private/RHICommandList.cpp b/Engine/Source/Runtime/RHI/Private/RHICommandList.cpp
--- a/Engine/Source/Runtime/RHI/Private/RHICommandList.cpp
+++ b/Engine/Source/Runtime/RHI/Private/RHICommandList.cpp
@@ -1033,6 +1033,9 @@ void FRHICommandListImmediate::QueueAsyncCompute(FRHIComputeCommandList& RHIComp
 	
 void FRHICommandListExecutor::WaitOnRHIThreadFence(FGraphEventRef& Fence)
 {
+	// Exclude RHIT waits from the RT critical path stat (these waits simply get longer if the RT is running faster, so we don't get useful results)
+	FThreadIdleStats::FScopeNonCriticalPath NonCriticalPathScope;
+
 	check(IsInRenderingThread());
 	if (Fence.GetReference() && !Fence->IsComplete())
 	{
@@ -1943,6 +1946,9 @@ FScopedCommandListWaitForTasks::~FScopedCommandListWaitForTasks()
 DECLARE_CYCLE_STAT(TEXT("Explicit wait for dispatch"), STAT_ExplicitWaitDispatch, STATGROUP_RHICMDLIST);
 void FRHICommandListBase::WaitForDispatch()
 {
+	// Exclude RHIT waits from the RT critical path stat (these waits simply get longer if the RT is running faster, so we don't get useful results)
+	FThreadIdleStats::FScopeNonCriticalPath NonCriticalPathScope;
+
 	check(IsImmediate() && IsInRenderingThread());
 	check(!AllOutstandingTasks.Num()); // dispatch before you get here
 	if (RenderThreadSublistDispatchTask.GetReference() && RenderThreadSublistDispatchTask->IsComplete())
@@ -2103,6 +2109,9 @@ void FRHICommandListImmediate::UnStallRHIThread()
 
 void FRHICommandListBase::WaitForRHIThreadTasks()
 {
+	// Exclude RHIT waits from the RT critical path stat (these waits simply get longer if the RT is running faster, so we don't get useful results)
+	FThreadIdleStats::FScopeNonCriticalPath NonCriticalPathScope;
+
 	check(IsImmediate() && IsInRenderingThread());
 	bool bAsyncSubmit = CVarRHICmdAsyncRHIThreadDispatch.GetValueOnRenderThread() > 0;
 	ENamedThreads::Type RenderThread_Local = ENamedThreads::GetRenderThread_Local();
diff --git a/Engine/Source/Runtime/RenderCore/Private/RenderingThread.cpp b/Engine/Source/Runtime/RenderCore/Private/RenderingThread.cpp
--- a/Engine/Source/Runtime/RenderCore/Private/RenderingThread.cpp
+++ b/Engine/Source/Runtime/RenderCore/Private/RenderingThread.cpp
@@ -295,6 +295,8 @@ uint32 GRenderThreadNumIdle[ERenderThreadIdleTypes::Num] = {0};
 uint32 GRenderThreadTime = 0;
 /** How many cycles the rhithread used (excluding idle time). */
 uint32 GRHIThreadTime = 0;
+/** How many cycles the renderthread used, including dependent wait time. */
+uint32 GRenderThreadTimeCriticalPath = 0;
 
 
 
diff --git a/Engine/Source/Runtime/RenderCore/Public/RenderCore.h b/Engine/Source/Runtime/RenderCore/Public/RenderCore.h
--- a/Engine/Source/Runtime/RenderCore/Public/RenderCore.h
+++ b/Engine/Source/Runtime/RenderCore/Public/RenderCore.h
@@ -292,6 +292,8 @@ extern RENDERCORE_API uint32 GRHIThreadTime;
 extern RENDERCORE_API uint32 GGameThreadTime;
 /** How many cycles it took to swap buffers to present the frame. */
 extern RENDERCORE_API uint32 GSwapBufferTime;
+/** How many cycles the renderthread used, including dependent wait time. */
+extern RENDERCORE_API uint32 GRenderThreadTimeCriticalPath;
 
 // shared by renderer and engine, compiles down to a constant in final release
 RENDERCORE_API int32 GetCVarForceLOD();
diff --git a/Engine/Source/Runtime/Renderer/Private/SceneRendering.cpp b/Engine/Source/Runtime/Renderer/Private/SceneRendering.cpp
--- a/Engine/Source/Runtime/Renderer/Private/SceneRendering.cpp
+++ b/Engine/Source/Runtime/Renderer/Private/SceneRendering.cpp
@@ -4551,6 +4551,13 @@ void FRendererModule::BeginRenderingViewFamilies(FCanvas* Canvas, TArrayView<FSc
 		}
 	}
 
+	ENQUEUE_RENDER_COMMAND(SetRtWaitCriticalPath)(
+		[](FRHICommandList& RHICmdList)
+		{
+			// Rendering is up and running now, so waits are considered part of the RT critical path
+			FThreadIdleStats::BeginCriticalPath();
+		});
+
 	ENQUEUE_RENDER_COMMAND(UpdateDeferredCachedUniformExpressions)(
 		[](FRHICommandList& RHICmdList)
 		{
diff --git a/Engine/Source/Runtime/SlateRHIRenderer/Private/SlateRHIRenderer.cpp b/Engine/Source/Runtime/SlateRHIRenderer/Private/SlateRHIRenderer.cpp
--- a/Engine/Source/Runtime/SlateRHIRenderer/Private/SlateRHIRenderer.cpp
+++ b/Engine/Source/Runtime/SlateRHIRenderer/Private/SlateRHIRenderer.cpp
@@ -40,6 +40,10 @@ DECLARE_CYCLE_STAT(TEXT("Slate RT: Rendering"), STAT_SlateRenderingRTTime, STATG
 
 DECLARE_CYCLE_STAT(TEXT("Slate RT: Draw Batches"), STAT_SlateRTDrawBatches, STATGROUP_Slate);
 
+DECLARE_CYCLE_STAT(TEXT("Total Render Thread time including dependent waits"), STAT_RenderThreadCriticalPath, STATGROUP_Threading);
+
+CSV_DEFINE_CATEGORY(RenderThreadIdle, true);
+
 DECLARE_GPU_DRAWCALL_STAT_NAMED(SlateUI, TEXT("Slate UI"));
 
 // Defines the maximum size that a slate viewport will create
@@ -89,6 +93,15 @@ TAutoConsoleVariable<int32> CVarShowSlateBatching(
 );
 #endif
 
+// RT stat including waits toggle. Off by default for historical tracking reasons
+TAutoConsoleVariable<int32> CVarRenderThreadTimeIncludesDependentWaits(
+	TEXT("r.RenderThreadTimeIncludesDependentWaits"),
+	0,
+	TEXT("0: RT stat only includes non-idle time, 1: RT stat includes dependent waits (matching RenderThreadTime_CriticalPath)"),
+	ECVF_Default
+);
+
+
 struct FSlateDrawWindowCommandParams
 {
 	FSlateRHIRenderer* Renderer;
@@ -1371,12 +1384,18 @@ void FSlateRHIRenderer::DrawWindow_RenderThread(FRHICommandListImmediate& RHICmd
 	GRenderThreadIdle[ERenderThreadIdleTypes::WaitingForAllOtherSleep] = RenderThread.Waits;
 	GRenderThreadIdle[ERenderThreadIdleTypes::WaitingForGPUPresent] += GSwapBufferTime;
 	GRenderThreadNumIdle[ERenderThreadIdleTypes::WaitingForGPUPresent]++;
-	RenderThread.Waits = 0;
 
 	SET_CYCLE_COUNTER(STAT_RenderingIdleTime_RenderThreadSleepTime, GRenderThreadIdle[0]);
 	SET_CYCLE_COUNTER(STAT_RenderingIdleTime_WaitingForGPUQuery, GRenderThreadIdle[1]);
 	SET_CYCLE_COUNTER(STAT_RenderingIdleTime_WaitingForGPUPresent, GRenderThreadIdle[2]);
 
+	// Set the RenderThreadIdle CSV stats
+	CSV_CUSTOM_STAT(RenderThreadIdle, Total, FPlatformTime::ToMilliseconds(RenderThread.Waits), ECsvCustomStatOp::Set);
+	CSV_CUSTOM_STAT(RenderThreadIdle, CriticalPath, FPlatformTime::ToMilliseconds(RenderThread.WaitsCriticalPath), ECsvCustomStatOp::Set);
+	CSV_CUSTOM_STAT(RenderThreadIdle, SwapBuffer, FPlatformTime::ToMilliseconds(GSwapBufferTime), ECsvCustomStatOp::Set);
+	CSV_CUSTOM_STAT(RenderThreadIdle, NonCriticalPath, FPlatformTime::ToMilliseconds(RenderThread.Waits - RenderThread.WaitsCriticalPath), ECsvCustomStatOp::Set);
+	CSV_CUSTOM_STAT(RenderThreadIdle, GPUQuery, FPlatformTime::ToMilliseconds(GRenderThreadIdle[1]), ECsvCustomStatOp::Set);
+
 	for (int32 Index = 0; Index < ERenderThreadIdleTypes::Num; Index++)
 	{
 		RenderThreadIdle += GRenderThreadIdle[Index];
@@ -1386,6 +1405,21 @@ void FSlateRHIRenderer::DrawWindow_RenderThread(FRHICommandListImmediate& RHICmd
 
 	SET_CYCLE_COUNTER(STAT_RenderingIdleTime, RenderThreadIdle);
 	GRenderThreadTime = (ThreadTime > RenderThreadIdle) ? (ThreadTime - RenderThreadIdle) : ThreadTime;
+
+	// Compute GRenderThreadTimeCriticalPath
+	uint32 RenderThreadNonCriticalPathIdle = RenderThreadIdle - RenderThread.WaitsCriticalPath;
+	GRenderThreadTimeCriticalPath = (ThreadTime > RenderThreadNonCriticalPathIdle) ? (ThreadTime - RenderThreadNonCriticalPathIdle) : ThreadTime;
+	SET_CYCLE_COUNTER(STAT_RenderThreadCriticalPath, GRenderThreadTimeCriticalPath);
+
+
+	if (CVarRenderThreadTimeIncludesDependentWaits.GetValueOnRenderThread())
+	{
+		// Optionally force the renderthread stat to include dependent waits
+		GRenderThreadTime = GRenderThreadTimeCriticalPath;
+	}
+
+	// Reset the idle stats
+	RenderThread.Reset();
 	
 	if (IsRunningRHIInSeparateThread())
 	{
