Background blur optimizations (followup of 20933909)

Instead of rendering to an 11-11-10 (in rec2020 space) and keeping a mask, we have the background blur write directly to the swapchain during the up-sample pass.
The 11-11-10 RT is kept but only for the background blur. BlitUIToHDRPS will take care of moving / blending in rec2020 and will output to the 11-11-10 in rec2020, which will be blurred just like in  20933909.
UpsampleMain will then take care of going back to the swapchain color space / format (and still erase the UI in the same region).
A copy of the swapchain will still happen, but right before UI composition. I will then remove that copy in a next CL for platforms that support UAV flag for swapchain.

Fix UpdateScissorRect that would reset rendertarget1 when being called. Use instead the current FRHIRenderPassInfo to keep the same RT layout when updating the stencilbuffer

diff --git a/Engine/Shaders/Private/CompositeUIPixelShader.usf b/Engine/Shaders/Private/CompositeUIPixelShader.usf
--- a/Engine/Shaders/Private/CompositeUIPixelShader.usf
+++ b/Engine/Shaders/Private/CompositeUIPixelShader.usf
@@ -63,7 +63,9 @@ void GenerateLUTPS(
 }
 
 ///////////////////////////////////////////////////////////////////////////////////////
+#if PLATFORM_SUPPORTS_RENDERTARGET_WRITE_MASK
 Texture2D<uint> UIWriteMaskTexture;
+#endif
 
 Texture2D UITexture;
 SamplerState UISampler;
@@ -135,80 +137,86 @@ void HDREditorConvert(
 	OutColor.rgb = SceneColor;
 }
 
-void Main(
-	FScreenVertexOutput Input,
-	out float4 OutColor : SV_Target0
-	)
+float3 ComposeUIAndScene(float3 SceneColor, float4 InUIColor, float InUILevel)
 {
-#if PLATFORM_SUPPORTS_RENDERTARGET_WRITE_MASK
-	uint CompositeUIMask = DecodeRTWriteMask(uint2(Input.Position.xy), UIWriteMaskTexture, 1);
 	BRANCH
-	if (CompositeUIMask == 0)
+	if (InUIColor.w > 0.f && InUIColor.w < 1.f)
 	{
-		discard;
-	}
-#else
-	uint CompositeUIMask = UIWriteMaskTexture.Load(uint3(Input.Position.xy,0));
-#endif
-
-	// In order to get consistent blending between HDR PQ & ScRGB, we convert UI color in both cases to rec2020, perform the blending in this color space, and
-	// move back to the output device space. Rec2020 is chosen over ScRGB because this is the way blending is currently being done in engine mode, and contains only positive values
-	const float3x3 sRGB_2_Rec2020 = mul(XYZ_2_Rec2020_MAT, sRGB_2_XYZ_MAT);
+		// Clamp gamut to sRGB as extended gamut colors bleeding into the UI can look funny
+		SceneColor = max(SceneColor, 0.f);
 
-	float4 UIColor = Texture2DSample(UITexture, UISampler, Input.UV);
+		// Tonemap HDR under transparent UI with a simple Reinhard to the max luminance of the UI
+		// This prevents HDR bleed through destroying UI legibility
+		// Rec2020 coefficients to compute luminance
+		float KR = 0.2627, KG = 0.678, KB = 0.0593;
+		float Luminance = dot(SceneColor, half3(KR, KG, KB)) / InUILevel;
+		float OutL = 1.f / (Luminance + 1.f);
 
-#if !PLATFORM_SUPPORTS_RENDERTARGET_WRITE_MASK
-	BRANCH
-	if (CompositeUIMask == 0 && UIColor.a < 1e-6f)
-	{
-		discard;
+		// Ease out remapping to avoid hard transitions where UI is near zero opacity
+		SceneColor *= lerp(1.f, OutL * InUILevel, InUIColor.w);
 	}
-#endif
 
-	{
-		UIColor.xyz = LinearizeColor(UIColor.xyz);
+	// Composite, assuming pre-multiplied alpha
+	return SceneColor * (1.f - InUIColor.w) + InUIColor.xyz * InUILevel;
+}
+
+// About sRGB_2_Rec2020 in ComputeHDRUIColor/ComputeHDRSceneColor: to get consistent blending between UI and scene color, regardless of the swapchain format choice, we perform blending in the rec2020 color space, and
+// move back to the output device space. Rec2020 is chosen over ScRGB because this is the way blending is currently being done in engine mode, and contains only positive values
+
+float4 ComputeHDRUIColor(float2 InputUV)
+{
+	const float3x3 sRGB_2_Rec2020 = mul(XYZ_2_Rec2020_MAT, sRGB_2_XYZ_MAT);
+	float4 UIColor = Texture2DSample(UITexture, UISampler, InputUV);
+	UIColor.xyz = LinearizeColor(UIColor.xyz);
 #if USE_LUT_FOR_UI
-		UIColor.xyz = ColorLookupTable(UIColor.xyz);
-		// PQ -> Linear
+	UIColor.xyz = ColorLookupTable(UIColor.xyz);
+	// PQ -> Linear
 #if !SCRGB_ENCODING
-		UIColor.xyz = ST2084ToLinear(UIColor.xyz);
+	UIColor.xyz = ST2084ToLinear(UIColor.xyz);
 #else
-		UIColor.xyz = ST2084ToScRGB(UIColor.xyz, OutputDevice, OutputMaxLuminance);
-		UIColor.xyz = mul(sRGB_2_Rec2020, UIColor.xyz) * UIBrightnessInNits;
+	UIColor.xyz = ST2084ToScRGB(UIColor.xyz, OutputDevice, OutputMaxLuminance);
+	UIColor.xyz = mul(sRGB_2_Rec2020, UIColor.xyz) * UIBrightnessInNits;
 #endif
 
 #else // #if USE_LUT_FOR_UI
-		UIColor.xyz = mul(sRGB_2_Rec2020, UIColor.xyz) * UIBrightnessInNits;
+	UIColor.xyz = mul(sRGB_2_Rec2020, UIColor.xyz) * UIBrightnessInNits;
 #endif // #if USE_LUT_FOR_UI
-	}
 
-	float3 SceneColor = Texture2DSample(SceneTexture, SceneSampler, Input.UV).xyz;
+	return UIColor;
+}
+
+float3 ComputeHDRSceneColor(float2 InputUV, SamplerState InSceneSampler)
+{
+	const float3x3 sRGB_2_Rec2020 = mul(XYZ_2_Rec2020_MAT, sRGB_2_XYZ_MAT);
+	float3 SceneColor = Texture2DSample(SceneTexture, InSceneSampler, InputUV).xyz;
 #if !SCRGB_ENCODING
-    // We don't use ST2084ToLinear here because HDRBackBufferConvertPS performed the conversion before
-	SceneColor.xyz = SceneColor.xyz * ScRGBScaleFactor;
+	SceneColor.xyz = ST2084ToLinear(SceneColor.xyz);
 #else
 	SceneColor.xyz = mul(sRGB_2_Rec2020, SceneColor.xyz) * ScRGBScaleFactor;
 #endif
+	return SceneColor;
+}
 
+
+void Main(
+	FScreenVertexOutput Input,
+	out float4 OutColor : SV_Target0
+	)
+{
+#if PLATFORM_SUPPORTS_RENDERTARGET_WRITE_MASK
+	uint CompositeUIMask = DecodeRTWriteMask(uint2(Input.Position.xy), UIWriteMaskTexture, 1);
 	BRANCH
-	if (UIColor.w > 0.f && UIColor.w < 1.f)
+	if (CompositeUIMask == 0)
 	{
-		// Clamp gamut to sRGB as extended gamut colors bleeding into the UI can look funny
-		SceneColor = max(SceneColor, 0.f);
+		discard;
+	}
+#endif
 
-		// Tonemap HDR under transparent UI with a simple Rienhard to the max luminance of the UI
-		// This prevents HDR bleed through destroying UI legibility
-		// Rec2020 coefficients to compute luminance
-		float KR = 0.2627, KG = 0.678, KB = 0.0593;
-		float Luminance = dot(SceneColor, half3(KR, KG, KB)) / UILevel;
-		float OutL = 1.f / (Luminance + 1.f);
+	float4 UIColor = ComputeHDRUIColor(Input.UV);
 
-		// Ease out remapping to avoid hard transitions where UI is near zero opacity
-		SceneColor *= lerp(1.f, OutL * UILevel, UIColor.w);
-	}
+	float3 SceneColor = ComputeHDRSceneColor(Input.UV, SceneSampler);
+	OutColor.xyz = ComposeUIAndScene(SceneColor.xyz, UIColor, UILevel);
 
-	// Composite, assuming pre-multiplied alpha
-	OutColor.xyz = SceneColor * (1.f - UIColor.w) + UIColor.xyz * UILevel;
 	OutColor.w = 1.f;
 
 #if !SCRGB_ENCODING
@@ -219,6 +227,7 @@ void Main(
 	OutColor.xyz = mul(Rec2020_2_sRGB, OutColor.xyz / ScRGBScaleFactor);
 #endif
 }
+
 #if DISPLAY_MAPPING_PS
 float4x4 TextureToOutputGamutMatrix;
 
@@ -248,37 +257,27 @@ void DisplayMappingPS(
 
 #if BLIT_UI_TO_HDR
 
-float4x4 SrgbToOutputMatrix;
+float2 UITextureSize;
 
 void BlitUIToHDRPS(
 	FScreenVertexOutput Input,
 	out float4 OutColor : SV_Target0
 	)
 {
+	float4 UIColor = ComputeHDRUIColor(Input.UV);
 #if PLATFORM_SUPPORTS_RENDERTARGET_WRITE_MASK
-	uint CompositeUIMask = DecodeRTWriteMask(uint2(Input.Position.xy), UIWriteMaskTexture, 1);
+	uint2 RTWriteMaskInputPosition = uint2(Input.UV * UITextureSize);
+	uint CompositeUIMask = DecodeRTWriteMask(RTWriteMaskInputPosition, UIWriteMaskTexture, 1);
 	BRANCH
 	if (CompositeUIMask == 0)
 	{
-		discard;
+		UIColor = float4(0,0,0,0);
 	}
 #endif
 
-	float4 UIColor = Texture2DSample(UITexture, UISampler, Input.UV);
-
-#if !PLATFORM_SUPPORTS_RENDERTARGET_WRITE_MASK
-	BRANCH
-	if (UIColor.a < 1e-3f)
-	{
-		discard;
-	}
-#endif
-
-	UIColor.rgb = LinearizeColor(UIColor.rgb) * UILevel * UIBrightnessInNits / ScRGBScaleFactor * UIColor.a;
-	UIColor.rgb = mul((float3x3)SrgbToOutputMatrix, UIColor.rgb);
-
-	OutColor.rgb = UIColor.rgb;
-	OutColor.a = UIColor.a;
+	float3 SceneColor = ComputeHDRSceneColor(Input.UV, UISampler);
+	OutColor.xyz = ComposeUIAndScene(SceneColor.xyz, UIColor, UILevel) / ScRGBScaleFactor;
+	OutColor.a = 1;
 }
 
 #endif
diff --git a/Engine/Shaders/Private/SlatePostProcessPixelShader.usf b/Engine/Shaders/Private/SlatePostProcessPixelShader.usf
--- a/Engine/Shaders/Private/SlatePostProcessPixelShader.usf
+++ b/Engine/Shaders/Private/SlatePostProcessPixelShader.usf
@@ -2,7 +2,7 @@
 
 #include "Common.ush"
 #include "SlateShaderCommon.ush"
-
+#include "GammaCorrectionCommon.ush"
 
 #define BILINEAR_FILTER_METHOD 1
 #define MAX_SAMPLES 127
@@ -119,9 +119,6 @@ void UpsampleMain(FScreenVertexOutput Input
 	, out float4 Color0 : SV_Target0
 #if OUTPUT_TO_UI_TARGET
 	, out float4 Color1 : SV_Target1
-#if !PLATFORM_SUPPORTS_RENDERTARGET_WRITE_MASK
-	, out uint Color2 : SV_Target2
-#endif
 #endif
 )
 {
@@ -131,13 +128,36 @@ void UpsampleMain(FScreenVertexOutput Input
 	const float4 CornerRadii = ShaderParams2;
 	const float4 FillColor = Texture2DSample(ElementTexture, ElementTextureSampler, Input.UV);
 	Color0 = GetRoundedBoxElementColorInternal(LocalSize, Map(UV, float2(0,0), ShaderParams.zw, float2(0,0),float2(1,1)), Thickness, CornerRadii, FillColor, FillColor);
+
 #if OUTPUT_TO_UI_TARGET
-	// still needed in case compositing needs PLATFORM_SUPPORTS_RENDERTARGET_WRITE_MASK
-	Color1 = float4(0, 0, 0, 0);
 
-#if !PLATFORM_SUPPORTS_RENDERTARGET_WRITE_MASK
-	Color2 = 255;
+	static const float3x3 XYZ_2_sRGB_MAT =
+	{
+		 3.2409699419, -1.5373831776, -0.4986107603,
+		-0.9692436363,  1.8759675015,  0.0415550574,
+		 0.0556300797, -0.2039769589,  1.0569715142,
+	};
+
+	static const float3x3 Rec2020_2_XYZ_MAT =
+	{
+		0.6369736, 0.1446172, 0.1688585,
+		0.2627066, 0.6779996, 0.0592938,
+		0.0000000, 0.0280728, 1.0608437
+	};
+
+
+	static const float ScRGBScaleFactor = 80.0f; // used to convert between ScRGB and nits
+	Color0.xyz *= ScRGBScaleFactor;
+
+#if SCRGB_ENCODING
+	const float3x3 Rec2020_2_sRGB = mul(XYZ_2_sRGB_MAT, Rec2020_2_XYZ_MAT);
+	Color0.xyz = mul(Rec2020_2_sRGB, Color0.xyz) / ScRGBScaleFactor;
+#else
+	Color0.xyz = LinearToST2084(Color0.xyz);
 #endif
 
+	// Erase the data that was previously in the UI now that the background blur is done
+	Color1 = float4(0, 0, 0, 0);
+
 #endif
 }
\ No newline at end of file
diff --git a/Engine/Source/Runtime/SlateRHIRenderer/Private/SlatePostProcessResource.cpp b/Engine/Source/Runtime/SlateRHIRenderer/Private/SlatePostProcessResource.cpp
--- a/Engine/Source/Runtime/SlateRHIRenderer/Private/SlatePostProcessResource.cpp
+++ b/Engine/Source/Runtime/SlateRHIRenderer/Private/SlatePostProcessResource.cpp
@@ -17,9 +17,9 @@ FSlatePostProcessResource::~FSlatePostProcessResource()
 
 }
 
-void FSlatePostProcessResource::Update(const FIntPoint& NewSize, const FTextureRHIRef& SourceTexture)
+void FSlatePostProcessResource::Update(const FIntPoint& NewSize, EPixelFormat RequestedPixelFormat)
 {
-	if(NewSize.X > RenderTargetSize.X || NewSize.Y > RenderTargetSize.Y || RenderTargetSize == FIntPoint::ZeroValue || RenderTargets.Num() == 0 || PixelFormat != SourceTexture->GetDesc().Format)
+	if(NewSize.X > RenderTargetSize.X || NewSize.Y > RenderTargetSize.Y || RenderTargetSize == FIntPoint::ZeroValue || RenderTargets.Num() == 0 || PixelFormat != RequestedPixelFormat)
 	{
 		if(!IsInitialized())
 		{
@@ -27,18 +27,18 @@ void FSlatePostProcessResource::Update(const FIntPoint& NewSize, const FTextureR
 		}
 
 		FIntPoint NewMaxSize(FMath::Max(NewSize.X, RenderTargetSize.X), FMath::Max(NewSize.Y, RenderTargetSize.Y));
-		ResizeTargets(NewMaxSize, SourceTexture);
+		ResizeTargets(NewMaxSize, RequestedPixelFormat);
 	}
 }
 
-void FSlatePostProcessResource::ResizeTargets(const FIntPoint& NewSize, const FTextureRHIRef& SourceTexture)
+void FSlatePostProcessResource::ResizeTargets(const FIntPoint& NewSize, EPixelFormat RequestedPixelFormat)
 {
 	check(IsInRenderingThread());
 
 	RenderTargets.Empty();
 
 	RenderTargetSize = NewSize;
-	PixelFormat = SourceTexture->GetDesc().Format;
+	PixelFormat = RequestedPixelFormat;
 	if (RenderTargetSize.X > 0 && RenderTargetSize.Y > 0)
 	{
 		for (int32 TexIndex = 0; TexIndex < RenderTargetCount; ++TexIndex)
diff --git a/Engine/Source/Runtime/SlateRHIRenderer/Private/SlatePostProcessResource.h b/Engine/Source/Runtime/SlateRHIRenderer/Private/SlatePostProcessResource.h
--- a/Engine/Source/Runtime/SlateRHIRenderer/Private/SlatePostProcessResource.h
+++ b/Engine/Source/Runtime/SlateRHIRenderer/Private/SlatePostProcessResource.h
@@ -22,7 +22,7 @@ public:
 	}
 
 	/** Performs per frame updates to this resource */
-	void Update(const FIntPoint& NewSize, const FTextureRHIRef& SourceTexture);
+	void Update(const FIntPoint& NewSize, EPixelFormat RequestedPixelFormat);
 
 	void CleanUp();
 
@@ -38,7 +38,7 @@ public:
 
 private:
 	/** Resizes targets to the new size */
-	void ResizeTargets(const FIntPoint& NewSize, const FTextureRHIRef& SourceTexture);
+	void ResizeTargets(const FIntPoint& NewSize, EPixelFormat RequestedPixelFormat);
 
 private:
 	TArray<FTexture2DRHIRef, TInlineAllocator<2>> RenderTargets;
diff --git a/Engine/Source/Runtime/SlateRHIRenderer/Private/SlatePostProcessor.cpp b/Engine/Source/Runtime/SlateRHIRenderer/Private/SlatePostProcessor.cpp
--- a/Engine/Source/Runtime/SlateRHIRenderer/Private/SlatePostProcessor.cpp
+++ b/Engine/Source/Runtime/SlateRHIRenderer/Private/SlatePostProcessor.cpp
@@ -30,9 +30,9 @@ FSlatePostProcessor::~FSlatePostProcessor()
 
 
 // Pixel shader to composite UI over HDR buffer
-class FBlitUIToHDRPS : public FGlobalShader
+class FBlitUIToHDRPSBase : public FGlobalShader
 {
-	DECLARE_SHADER_TYPE(FBlitUIToHDRPS, Global);
+	DECLARE_TYPE_LAYOUT(FBlitUIToHDRPSBase, NonVirtual);
 public:
 
 	static bool ShouldCompilePermutation(const FGlobalShaderPermutationParameters& Parameters)
@@ -40,23 +40,25 @@ public:
 		return IsFeatureLevelSupported(Parameters.Platform, ERHIFeatureLevel::SM5) && (RHISupportsGeometryShaders(Parameters.Platform) || RHISupportsVertexShaderLayer(Parameters.Platform));
 	}
 
-	FBlitUIToHDRPS(const ShaderMetaType::CompiledShaderInitializerType& Initializer) :
+	FBlitUIToHDRPSBase(const ShaderMetaType::CompiledShaderInitializerType& Initializer) :
 		FGlobalShader(Initializer)
 	{
+		HDRTexture.Bind(Initializer.ParameterMap, TEXT("SceneTexture"));
 		UITexture.Bind(Initializer.ParameterMap, TEXT("UITexture"));
 		UIWriteMaskTexture.Bind(Initializer.ParameterMap, TEXT("UIWriteMaskTexture"));
 		UISampler.Bind(Initializer.ParameterMap, TEXT("UISampler"));
 		UILevel.Bind(Initializer.ParameterMap, TEXT("UILevel"));
-		SrgbToOutputMatrix.Bind(Initializer.ParameterMap, TEXT("SrgbToOutputMatrix"));
+		UITextureSize.Bind(Initializer.ParameterMap, TEXT("UITextureSize"));
 	}
-	FBlitUIToHDRPS() {}
+	FBlitUIToHDRPSBase() = default;
 
-	void SetParameters(FRHICommandList& RHICmdList, FRHITexture* UITextureRHI, FRHITexture* UITextureWriteMaskRHI, const FMatrix44f& InSrgbToOutputMatrix)
+	void SetParameters(FRHICommandList& RHICmdList, FRHITexture* UITextureRHI, FRHITexture* HDRTextureRHI, FRHITexture* UITextureWriteMaskRHI, const FVector2f& InUITextureSize)
 	{
-		SetTextureParameter(RHICmdList, RHICmdList.GetBoundPixelShader(), UITexture, UISampler, TStaticSamplerState<SF_Point>::GetRHI(), UITextureRHI);
+		SetTextureParameter(RHICmdList, RHICmdList.GetBoundPixelShader(), HDRTexture, UISampler, TStaticSamplerState<SF_Bilinear>::GetRHI(), HDRTextureRHI);
+		SetTextureParameter(RHICmdList, RHICmdList.GetBoundPixelShader(), UITexture, UISampler, TStaticSamplerState<SF_Bilinear>::GetRHI(), UITextureRHI);
 		static auto CVarHDRUILevel = IConsoleManager::Get().FindConsoleVariable(TEXT("r.HDR.UI.Level"));
 		SetShaderValue(RHICmdList, RHICmdList.GetBoundPixelShader(), UILevel, CVarHDRUILevel ? CVarHDRUILevel->GetFloat() : 1.0f);
-		SetShaderValue(RHICmdList, RHICmdList.GetBoundPixelShader(), SrgbToOutputMatrix, InSrgbToOutputMatrix);
+		SetShaderValue(RHICmdList, RHICmdList.GetBoundPixelShader(), UITextureSize, InUITextureSize);
 		if (UITextureWriteMaskRHI != nullptr)
 		{
 			SetTextureParameter(RHICmdList, RHICmdList.GetBoundPixelShader(), UIWriteMaskTexture, UITextureWriteMaskRHI);
@@ -80,22 +82,41 @@ public:
 	}
 
 private:
+	LAYOUT_FIELD(FShaderResourceParameter, HDRTexture);
 	LAYOUT_FIELD(FShaderResourceParameter, UITexture);
 	LAYOUT_FIELD(FShaderResourceParameter, UISampler);
 	LAYOUT_FIELD(FShaderResourceParameter, UIWriteMaskTexture);
-	LAYOUT_FIELD(FShaderParameter, SrgbToOutputMatrix);
 	LAYOUT_FIELD(FShaderParameter, UILevel);
+	LAYOUT_FIELD(FShaderParameter, UITextureSize);
+};
+
+IMPLEMENT_TYPE_LAYOUT(FBlitUIToHDRPSBase);
+
+template<uint32 EncodingType>
+class FBlitUIToHDRPS : public FBlitUIToHDRPSBase
+{
+	DECLARE_SHADER_TYPE(FBlitUIToHDRPS, Global);
+public:
+	FBlitUIToHDRPS() = default;
+	FBlitUIToHDRPS(const ShaderMetaType::CompiledShaderInitializerType& Initializer) : FBlitUIToHDRPSBase(Initializer) {}
+
+	static void ModifyCompilationEnvironment(const FGlobalShaderPermutationParameters& Parameters, FShaderCompilerEnvironment& OutEnvironment)
+	{
+		FBlitUIToHDRPSBase::ModifyCompilationEnvironment(Parameters, OutEnvironment);
+		OutEnvironment.SetDefine(TEXT("SCRGB_ENCODING"), EncodingType);
+	}
 };
 
-IMPLEMENT_SHADER_TYPE(, FBlitUIToHDRPS, FBlitUIToHDRPS::GetSourceFilename(), FBlitUIToHDRPS::GetFunctionName(), SF_Pixel);
+IMPLEMENT_SHADER_TYPE(template<>, FBlitUIToHDRPS<0>, FBlitUIToHDRPS::GetSourceFilename(), FBlitUIToHDRPS::GetFunctionName(), SF_Pixel);
+IMPLEMENT_SHADER_TYPE(template<>, FBlitUIToHDRPS<1>, FBlitUIToHDRPS::GetSourceFilename(), FBlitUIToHDRPS::GetFunctionName(), SF_Pixel);
 
-static void BlitUIToHDRScene(FRHICommandListImmediate& RHICmdList, IRendererModule& RendererModule, const FPostProcessRectParams& RectParams)
+static void BlitUIToHDRScene(FRHICommandListImmediate& RHICmdList, IRendererModule& RendererModule, const FPostProcessRectParams& RectParams, FTexture2DRHIRef DestTexture, const FIntPoint& AllocatedSize)
 {
 	SCOPED_DRAW_EVENT(RHICmdList, SlatePostProcessBlitUIToHDR);
 
-	FRHITexture* SourceTexture = RectParams.UITarget->GetRHI();
+	FRHITexture* UITexture = RectParams.UITarget->GetRHI();
 
-	RHICmdList.Transition(FRHITransitionInfo(SourceTexture, ERHIAccess::Unknown, ERHIAccess::SRVMask));
+	RHICmdList.Transition(FRHITransitionInfo(UITexture, ERHIAccess::Unknown, ERHIAccess::SRVMask));
 	TRefCountPtr<IPooledRenderTarget> UITargetRTMask;
 	if (RHISupportsRenderTargetWriteMask(GMaxRHIShaderPlatform))
 	{
@@ -119,17 +140,25 @@ static void BlitUIToHDRScene(FRHICommandListImmediate& RHICmdList, IRendererModu
 	FGlobalShaderMap* ShaderMap = GetGlobalShaderMap(GMaxRHIFeatureLevel);
 	TShaderMapRef<FScreenVS> VertexShader(ShaderMap);
 
-	FTexture2DRHIRef DestTexture = RectParams.SourceTexture;
+	bool bIsSCRGB = (RectParams.SourceTexture->GetDesc().Format == PF_FloatRGBA);
 
-	TShaderMapRef<FBlitUIToHDRPS> PixelShader(ShaderMap);
+	TShaderRef<FBlitUIToHDRPSBase> PixelShader;  
+	if (bIsSCRGB)
+	{
+		PixelShader = TShaderMapRef<FBlitUIToHDRPS<1> >(ShaderMap);
+	}
+	else
+	{
+		PixelShader = TShaderMapRef<FBlitUIToHDRPS<0> >(ShaderMap);
+	}
 
 	RHICmdList.Transition(FRHITransitionInfo(DestTexture, ERHIAccess::Unknown, ERHIAccess::RTV));
 
 	const FVector2f InvSrcTextureSize(1.f / SrcTextureWidth, 1.f / SrcTextureHeight);
 
-	// Add some guard band to ensure blur will reach these pixels. It will overwrite pixels below, but these are going to be composited at the end anyway
-	const FVector2f UVStart = FVector2f(DestRect.Left - 10.0f, DestRect.Top - 10.0f) * InvSrcTextureSize;
-	const FVector2f UVEnd = FVector2f(DestRect.Right + 10.0f, DestRect.Bottom + 10.0f) * InvSrcTextureSize;
+	// no guard band is actually needed because GaussianBlurMain already takes ensure that sampling happens inside the rectangle
+	const FVector2f UVStart = FVector2f(DestRect.Left - 0.0f, DestRect.Top - 0.0f) * InvSrcTextureSize;
+	const FVector2f UVEnd = FVector2f(DestRect.Right + 0.0f, DestRect.Bottom + 0.0f) * InvSrcTextureSize;
 	const FVector2f SizeUV = UVEnd - UVStart;
 
 	RHICmdList.SetViewport(0, 0, 0, SrcTextureWidth, SrcTextureHeight, 0.0f);
@@ -140,7 +169,7 @@ static void BlitUIToHDRScene(FRHICommandListImmediate& RHICmdList, IRendererModu
 	{
 		FGraphicsPipelineStateInitializer GraphicsPSOInit;
 		RHICmdList.ApplyCachedRenderTargets(GraphicsPSOInit);
-		GraphicsPSOInit.BlendState = TStaticBlendState<CW_RGB, BO_Add, BF_One, BF_InverseSourceAlpha, BO_Add, BF_One, BF_InverseSourceAlpha>::GetRHI();
+		GraphicsPSOInit.BlendState = TStaticBlendState<CW_RGB>::GetRHI();
 		GraphicsPSOInit.RasterizerState = TStaticRasterizerState<>::GetRHI();
 		GraphicsPSOInit.DepthStencilState = TStaticDepthStencilState<false, CF_Always>::GetRHI();
 
@@ -152,20 +181,15 @@ static void BlitUIToHDRScene(FRHICommandListImmediate& RHICmdList, IRendererModu
 
 		FRHITexture* UITargetRTMaskTexture = RHISupportsRenderTargetWriteMask(GMaxRHIShaderPlatform) ? UITargetRTMask->GetRHI() : nullptr;
 
-		FMatrix44f Srgb_2_XYZ = GamutToXYZMatrix(EDisplayColorGamut::sRGB_D65);
-		FMatrix44f XYZ_2_Output = XYZToGamutMatrix(RectParams.HDRDisplayColorGamut);
-		// note: we use mul(m,v) instead of mul(v,m) in the shaders for color conversions which is why matrix multiplication is reversed compared to what we usually do
-		FMatrix44f CombinedMatrix = XYZ_2_Output * Srgb_2_XYZ;
-
-		PixelShader->SetParameters(RHICmdList, SourceTexture, UITargetRTMaskTexture, CombinedMatrix);
+		PixelShader->SetParameters(RHICmdList, UITexture, RectParams.SourceTexture, UITargetRTMaskTexture, FVector2f(SrcTextureWidth, SrcTextureHeight));
 
 		RendererModule.DrawRectangle(
 			RHICmdList,
-			UVStart.X * SrcTextureWidth, UVStart.Y * SrcTextureHeight,
-			SizeUV.X * SrcTextureWidth, SizeUV.Y * SrcTextureHeight,
+			0, 0,
+			(float)AllocatedSize.X, (float)AllocatedSize.Y,
 			UVStart.X, UVStart.Y,
 			SizeUV.X, SizeUV.Y,
-			FIntPoint(SrcTextureWidth, SrcTextureHeight),
+			DestTexture->GetDesc().Extent,
 			FIntPoint(1, 1),
 			VertexShader,
 			EDRF_Default);
@@ -198,15 +222,23 @@ void FSlatePostProcessor::BlurRect(FRHICommandListImmediate& RHICmdList, IRender
 
 	const FIntPoint DownsampleSize = RequiredSize;
 
-	IntermediateTargets->Update(RequiredSize, RectParams.SourceTexture);
+	EPixelFormat IntermediatePixelFormat = RectParams.SourceTexture->GetDesc().Format;
+
+	bool bIsHDRSource = RectParams.UITarget.IsValid() && RectParams.UITarget->GetRHI() != RectParams.SourceTexture.GetReference();
+	if (bIsHDRSource)
+	{
+		IntermediatePixelFormat = PF_FloatR11G11B10;
+	}
+
+	IntermediateTargets->Update(RequiredSize, IntermediatePixelFormat);
 
-	if (RectParams.UITarget.IsValid() && RectParams.UITarget->GetRHI() != RectParams.SourceTexture.GetReference())
+	if (bIsHDRSource)
 	{
 		// in HDR mode, we are going to blur SourceTexture, but still need to take into account the UI already rendered. Blit UI into HDR target
-		BlitUIToHDRScene(RHICmdList, RendererModule, RectParams);
+		BlitUIToHDRScene(RHICmdList, RendererModule, RectParams, IntermediateTargets->GetRenderTarget(0), RequiredSize);
 	}
 
-	if(bDownsample)
+	else if(bDownsample)
 	{
 		DownsampleRect(RHICmdList, RendererModule, RectParams, DownsampleSize);
 	}
@@ -248,7 +280,7 @@ void FSlatePostProcessor::BlurRect(FRHICommandListImmediate& RHICmdList, IRender
 		// First pass render to the render target with the post process fx
 		if (PassIndex == 0)
 		{
-			FTexture2DRHIRef SourceTexture = bDownsample ? IntermediateTargets->GetRenderTarget(0) : RectParams.SourceTexture;
+			FTexture2DRHIRef SourceTexture = (bDownsample || bIsHDRSource) ? IntermediateTargets->GetRenderTarget(0) : RectParams.SourceTexture;
 			FTexture2DRHIRef DestTexture = IntermediateTargets->GetRenderTarget(1);
 
 			RHICmdList.Transition(FRHITransitionInfo(SourceTexture, ERHIAccess::Unknown, ERHIAccess::SRVGraphics));
@@ -268,7 +300,7 @@ void FSlatePostProcessor::BlurRect(FRHICommandListImmediate& RHICmdList, IRender
 				PixelShader->SetWeightsAndOffsets(RHICmdList, WeightsAndOffsets, SampleCount);
 				PixelShader->SetTexture(RHICmdList, SourceTexture, BilinearClamp);
 
-				if (bDownsample)
+				if (bDownsample || bIsHDRSource)
 				{
 					PixelShader->SetUVBounds(RHICmdList, FVector4f(FVector2f::ZeroVector, FVector2f((float)DownsampleSize.X / DestTextureWidth, (float)DownsampleSize.Y / DestTextureHeight) - HalfTexelOffset));
 					PixelShader->SetBufferSizeAndDirection(RHICmdList, InvBufferSize, FVector2f(1.f, 0.f));
@@ -360,7 +392,7 @@ void FSlatePostProcessor::ColorDeficiency(FRHICommandListImmediate& RHICmdList,
 	FIntPoint DestRectSize = RectParams.DestRect.GetSize().IntPoint();
 	FIntPoint RequiredSize = DestRectSize;
 
-	IntermediateTargets->Update(RequiredSize, RectParams.SourceTexture);
+	IntermediateTargets->Update(RequiredSize, RectParams.SourceTexture->GetDesc().Format);
 
 #if 1
 	FSamplerStateRHIRef PointClamp = TStaticSamplerState<SF_Point, AM_Clamp, AM_Clamp, AM_Clamp>::GetRHI();
@@ -561,6 +593,7 @@ void FSlatePostProcessor::UpsampleRect(FRHICommandListImmediate& RHICmdList, IRe
 	FRHIRenderPassInfo RPInfo(DestTexture, ERenderTargetActions::Load_Store);
 
 	bool bHasMRT = false;
+	bool bIsSCRGB = false;
 
 	FRHITexture* UITargetTexture = Params.UITarget.IsValid() ? Params.UITarget->GetRHI() : nullptr;
 
@@ -569,34 +602,35 @@ void FSlatePostProcessor::UpsampleRect(FRHICommandListImmediate& RHICmdList, IRe
 		RPInfo.ColorRenderTargets[1].RenderTarget = UITargetTexture;
 		RPInfo.ColorRenderTargets[1].ArraySlice = -1;
 		RPInfo.ColorRenderTargets[1].Action = ERenderTargetActions::Load_Store;
-		if (!RHISupportsRenderTargetWriteMask(GMaxRHIShaderPlatform))
-		{
-			RPInfo.ColorRenderTargets[2].RenderTarget = Params.UITargetMask->GetRHI();
-			RPInfo.ColorRenderTargets[2].ArraySlice = -1;
-			RPInfo.ColorRenderTargets[2].Action = ERenderTargetActions::Load_Store;
-		}
 		bHasMRT = true;
+		bIsSCRGB = (DestTexture->GetDesc().Format == PF_FloatRGBA);
 	}
 
 	RHICmdList.BeginRenderPass(RPInfo, TEXT("UpsampleRect"));
 	{
 		RHICmdList.ApplyCachedRenderTargets(GraphicsPSOInit);
-
 		if (Params.RestoreStateFunc)
 		{
 			// This can potentially end and restart a renderpass.
 			// #todo refactor so that we only start one renderpass here. Right now RestoreStateFunc may call UpdateScissorRect which requires an open renderpass.
-			Params.RestoreStateFunc(RHICmdList, GraphicsPSOInit);
+			Params.RestoreStateFunc(RHICmdList, GraphicsPSOInit, RPInfo);
 		}
 
 		TShaderRef<FSlateElementPS> PixelShader;
 		if (bHasMRT)
 		{
-			PixelShader = TShaderMapRef<FSlatePostProcessUpsamplePS<true> >(ShaderMap);
+			if (bIsSCRGB)
+			{
+				PixelShader = TShaderMapRef<FSlatePostProcessUpsamplePS<ESlatePostProcessUpsamplePSPermutation::HDR_SCRGB> >(ShaderMap);
+			}
+			else
+			{
+				PixelShader = TShaderMapRef<FSlatePostProcessUpsamplePS<ESlatePostProcessUpsamplePSPermutation::HDR_PQ10> >(ShaderMap);
+			}
 		}
 		else
 		{
-			PixelShader = TShaderMapRef<FSlatePostProcessUpsamplePS<false> >(ShaderMap);
+			PixelShader = TShaderMapRef<FSlatePostProcessUpsamplePS<ESlatePostProcessUpsamplePSPermutation::SDR> >(ShaderMap);
 		}
 
 		GraphicsPSOInit.BoundShaderState.VertexDeclarationRHI = GFilterVertexDeclaration.VertexDeclarationRHI;
diff --git a/Engine/Source/Runtime/SlateRHIRenderer/Private/SlatePostProcessor.h b/Engine/Source/Runtime/SlateRHIRenderer/Private/SlatePostProcessor.h
--- a/Engine/Source/Runtime/SlateRHIRenderer/Private/SlatePostProcessor.h
+++ b/Engine/Source/Runtime/SlateRHIRenderer/Private/SlatePostProcessor.h
@@ -17,9 +17,8 @@ struct FPostProcessRectParams
 	FSlateRect DestRect;
 	FVector4f CornerRadius;
 	FIntPoint SourceTextureSize;
-	TFunction<void(FRHICommandListImmediate&, FGraphicsPipelineStateInitializer&)> RestoreStateFunc;
+	TFunction<void(FRHICommandListImmediate&, FGraphicsPipelineStateInitializer&, FRHIRenderPassInfo&)> RestoreStateFunc;
 	TRefCountPtr<IPooledRenderTarget> UITarget; // not using FTextureRHIRef because we want to be able to use FRenderTargetWriteMask::Decode
-	TRefCountPtr<IPooledRenderTarget> UITargetMask; // not using FTextureRHIRef because we want to be able to use FRenderTargetWriteMask::Decode
 	uint32 StencilRef{};
 	EDisplayColorGamut HDRDisplayColorGamut;
 };
diff --git a/Engine/Source/Runtime/SlateRHIRenderer/Private/SlateRHIRenderer.cpp b/Engine/Source/Runtime/SlateRHIRenderer/Private/SlateRHIRenderer.cpp
--- a/Engine/Source/Runtime/SlateRHIRenderer/Private/SlateRHIRenderer.cpp
+++ b/Engine/Source/Runtime/SlateRHIRenderer/Private/SlateRHIRenderer.cpp
@@ -132,8 +132,6 @@ void FViewportInfo::ReleaseResource()
 {
 	FRenderResource::ReleaseResource();
 	UITargetRT.SafeRelease();
-	UITargetRTMask.SafeRelease();
-	HDRSourceRT.SafeRelease();
 }
 
 void FViewportInfo::ConditionallyUpdateDepthBuffer(bool bInRequiresStencilTest, uint32 InWidth, uint32 InHeight)
@@ -822,7 +820,6 @@ void RenderSlateBatch(FTexture2DRHIRef SlateRenderTarget, bool bClear, bool bIsH
 				if (ViewportInfo.bSceneHDREnabled && !bIsHDR)
 				{
 					RenderParams.UITarget = ViewportInfo.UITargetRT;
-					RenderParams.UITargetMask = ViewportInfo.UITargetRTMask;
 				}
 				RenderingPolicy->SetUseGammaCorrection(!bIsHDR);
 
@@ -930,29 +927,8 @@ void FSlateRHIRenderer::DrawWindow_RenderThread(FRHICommandListImmediate& RHICmd
 			const uint32 ViewportWidth = (ViewportRT) ? ViewportRT->GetSizeX() : ViewportInfo.Width;
 			const uint32 ViewportHeight = (ViewportRT) ? ViewportRT->GetSizeY() : ViewportInfo.Height;
 
-			EPixelFormat HDRSourceRTExpectedFormat = BackBuffer->GetFormat();
-			if (bCompositeUI)
-			{
-				// Some effects like slate posprocess blur get quite expensive if we have to do PQ->linear for every texture sample. Instead, work with a format where values are stored linearly
-				// Do not test BackBuffer->GetFormat() == PF_A2B10G10R10 because some platforms do not return this PF with ST2084
-				if (ViewportInfo.HDRDisplayOutputFormat == EDisplayOutputFormat::HDR_ACES_1000nit_ST2084 || ViewportInfo.HDRDisplayOutputFormat == EDisplayOutputFormat::HDR_ACES_2000nit_ST2084)
-				{
-					// Prefer small float format to improve performance for slate's background blur. Banding might appear when choosing this format, but UITargetRTMask allows to discard pixels that have
-                    // not been touched by the UI
-					if (GPixelFormats[PF_FloatR11G11B10].Supported)
-					{
-						HDRSourceRTExpectedFormat = PF_FloatR11G11B10;
-					}
-					else
-					{
-						HDRSourceRTExpectedFormat = PF_FloatRGBA;
-					}
-				}
-			}
-
 			// Check to see that targets are up-to-date
-			if (bCompositeUI && (!ViewportInfo.UITargetRT || ViewportInfo.UITargetRT->GetRHI()->GetSizeX() != ViewportWidth || ViewportInfo.UITargetRT->GetRHI()->GetSizeY() != ViewportHeight
-				|| !ViewportInfo.HDRSourceRT || ViewportInfo.HDRSourceRT->GetRHI()->GetFormat() != HDRSourceRTExpectedFormat))
+			if (bCompositeUI && (!ViewportInfo.UITargetRT || ViewportInfo.UITargetRT->GetRHI()->GetSizeX() != ViewportWidth || ViewportInfo.UITargetRT->GetRHI()->GetSizeY() != ViewportHeight))
 			{
 				// Composition buffers
 				{
@@ -968,18 +944,6 @@ void FSlateRHIRenderer::DrawWindow_RenderThread(FRHICommandListImmediate& RHICmd
 						true));
 
 					GRenderTargetPool.FindFreeElement(RHICmdList, Desc, ViewportInfo.UITargetRT, TEXT("UITargetRT"));
-
-					Desc.Format = HDRSourceRTExpectedFormat;
-					GRenderTargetPool.FindFreeElement(RHICmdList, Desc, ViewportInfo.HDRSourceRT, TEXT("HDRSourceRT"));
-
-                    // Because we may use small float format for storing scene color data, keep a mask to discard pixels that were untouched by the UI. Note: relying only on the UI Rendertarget's alpha does 
-                    // not work because of FSlatePostProcessor::BlurRect. It composes UI to scene / blur it / remove the UI that was there before the blur. After that, UI pixels are merged with the scene pixels
-                    // This RT will keep track of where the blur happened, and will be used with UITargetRT's alpha channel
-					if (!RHISupportsRenderTargetWriteMask(GMaxRHIShaderPlatform))
-					{
-						Desc.Format = PF_R8_UINT;
-						GRenderTargetPool.FindFreeElement(RHICmdList, Desc, ViewportInfo.UITargetRTMask, TEXT("UITargetRTMask"));
-					}
 				}
 
 				// LUT
@@ -1062,70 +1026,7 @@ void FSlateRHIRenderer::DrawWindow_RenderThread(FRHICommandListImmediate& RHICmd
 					RenderSlateBatch(UITargetHDRRTRHI, /*bClear*/ false, /*bIsHDR*/ true, ViewportInfo, ViewMatrix, BatchDataHDR, RHICmdList, ViewportWidth, ViewportHeight, DrawCommandParams, RenderingPolicy, PostProcessBuffer);
 			    }
 
-				// Grab HDR backbuffer
-				if (ViewportInfo.HDRSourceRT->GetDesc().Format == FinalBuffer->GetDesc().Format)
-				{
-					TransitionAndCopyTexture(RHICmdList, FinalBuffer, ViewportInfo.HDRSourceRT->GetRHI(), {});
-				}
-				else
-				{
-					// Do not test FinalBuffer->GetDesc().Format == PF_A2B10G10R10 because some platforms do not return this PF even with ST2084
-					//ensure(FinalBuffer->GetDesc().Format == PF_A2B10G10R10);
-					ensure(ViewportInfo.HDRDisplayOutputFormat == EDisplayOutputFormat::HDR_ACES_1000nit_ST2084 || ViewportInfo.HDRDisplayOutputFormat == EDisplayOutputFormat::HDR_ACES_2000nit_ST2084);
-					ensure(ViewportInfo.HDRSourceRT->GetDesc().Format == PF_FloatR11G11B10 || ViewportInfo.HDRSourceRT->GetDesc().Format == PF_FloatRGBA);
-
-					// Convert Scene Color in PQ to Linear so that Slate can perform post processing on it without any conversion
-					{
-						auto ShaderMap = GetGlobalShaderMap(GMaxRHIFeatureLevel);
-						RHICmdList.Transition({ FRHITransitionInfo(ViewportInfo.HDRSourceRT->GetRHI(), ERHIAccess::Unknown, ERHIAccess::RTV),
-							                    FRHITransitionInfo(FinalBuffer, ERHIAccess::Unknown, ERHIAccess::SRVGraphics)
-							                  });
-						FRHIRenderPassInfo RPInfo(ViewportInfo.HDRSourceRT->GetRHI(), ERenderTargetActions::Load_Store);
-						RHICmdList.BeginRenderPass(RPInfo, TEXT("ConvertBackBuffer"));
-						{
-							FGraphicsPipelineStateInitializer GraphicsPSOInit;
-							RHICmdList.ApplyCachedRenderTargets(GraphicsPSOInit);
-							GraphicsPSOInit.BlendState = TStaticBlendState<>::GetRHI();
-							GraphicsPSOInit.RasterizerState = TStaticRasterizerState<>::GetRHI();
-							GraphicsPSOInit.DepthStencilState = TStaticDepthStencilState<false, CF_Always>::GetRHI();
-
-							TShaderMapRef<FScreenVS> VertexShader(ShaderMap);
-							TShaderMapRef<FHDRBackBufferConvertPS> PixelShader(ShaderMap);
-
-							GraphicsPSOInit.BoundShaderState.VertexDeclarationRHI = GFilterVertexDeclaration.VertexDeclarationRHI;
-							GraphicsPSOInit.BoundShaderState.VertexShaderRHI = VertexShader.GetVertexShader();
-							GraphicsPSOInit.BoundShaderState.PixelShaderRHI = PixelShader.GetPixelShader();
-							GraphicsPSOInit.PrimitiveType = PT_TriangleList;
-
-							SetGraphicsPipelineState(RHICmdList, GraphicsPSOInit, 0);
-
-							PixelShader->SetParameters(RHICmdList, FinalBuffer);
-
-							static const FName RendererModuleName("Renderer");
-							IRendererModule& RendererModule = FModuleManager::GetModuleChecked<IRendererModule>(RendererModuleName);
-							RendererModule.DrawRectangle(
-								RHICmdList,
-								0.f, 0.f,
-								(float)ViewportWidth, (float)ViewportHeight,
-								0.f, 0.f,
-								(float)ViewportWidth, (float)ViewportHeight,
-								FIntPoint(ViewportWidth, ViewportHeight),
-								FIntPoint(ViewportWidth, ViewportHeight),
-								VertexShader,
-								EDRF_UseTriangleOptimization);
-						}
-						RHICmdList.EndRenderPass();
-					}
-				}
-
-				if (!RHISupportsRenderTargetWriteMask(GMaxRHIShaderPlatform) && ViewportInfo.UITargetRTMask.IsValid())
-				{
-					FRHIRenderPassInfo RPInfo(ViewportInfo.UITargetRTMask->GetRHI(), ERenderTargetActions::Clear_Store);
-					RHICmdList.BeginRenderPass(RPInfo, TEXT("Clear_UITargetRTMask"));
-					RHICmdList.EndRenderPass();
-				}
-
-		    	PostProcessBuffer = ViewportInfo.HDRSourceRT->GetRHI();
+				PostProcessBuffer = FinalBuffer;
             }
 
 			FSlateBatchData& BatchData = WindowElementList.GetBatchData();
@@ -1200,9 +1101,24 @@ void FSlateRHIRenderer::DrawWindow_RenderThread(FRHICommandListImmediate& RHICmd
 						FRenderTargetWriteMask::Decode(RHICmdList, ShaderMap, RenderTargets, ViewportInfo.UITargetRTMask, TexCreate_None, TEXT("UIRTWriteMask"));
 					}
 
+
+					TRefCountPtr<IPooledRenderTarget> FinalBufferCopy;
+					{
+						FPooledRenderTargetDesc Desc(FPooledRenderTargetDesc::Create2DDesc(FIntPoint(ViewportWidth, ViewportHeight),
+							FinalBuffer->GetFormat(),
+							FClearValueBinding::Transparent,
+							TexCreate_None,
+							TexCreate_ShaderResource,
+							false,
+							1,
+							true,
+							true));
+						GRenderTargetPool.FindFreeElement(RHICmdList, Desc, FinalBufferCopy, TEXT("FinalBufferCopy"));
+					}
+					TransitionAndCopyTexture(RHICmdList, FinalBuffer, FinalBufferCopy->GetRHI(), {});
+
 					RHICmdList.Transition({
-						FRHITransitionInfo(FinalBuffer, ERHIAccess::Unknown, ERHIAccess::RTV),
-						FRHITransitionInfo(ViewportInfo.HDRSourceRT->GetRHI(), ERHIAccess::Unknown, ERHIAccess::SRVGraphics)
+						FRHITransitionInfo(FinalBuffer, ERHIAccess::Unknown, ERHIAccess::RTV)
 					});
 					FRHIRenderPassInfo RPInfo(FinalBuffer, ERenderTargetActions::Load_Store);
 					RHICmdList.BeginRenderPass(RPInfo, TEXT("SlateComposite"));
@@ -1228,7 +1144,7 @@ void FSlateRHIRenderer::DrawWindow_RenderThread(FRHICommandListImmediate& RHICmd
 
 							SetGraphicsPipelineState(RHICmdList, GraphicsPSOInit, 0);
 
-							PixelShader->SetParameters(RHICmdList, ViewportInfo.UITargetRT->GetRHI(), UITargetRTMaskTexture, ViewportInfo.HDRSourceRT->GetRHI(), ViewportInfo.ColorSpaceLUT);
+							PixelShader->SetParameters(RHICmdList, ViewportInfo.UITargetRT->GetRHI(), UITargetRTMaskTexture, FinalBufferCopy->GetRHI(), ViewportInfo.ColorSpaceLUT);
 						}
 						else
 						{
@@ -1242,7 +1158,7 @@ void FSlateRHIRenderer::DrawWindow_RenderThread(FRHICommandListImmediate& RHICmd
 
 							SetGraphicsPipelineState(RHICmdList, GraphicsPSOInit, 0);
 
-							PixelShader->SetParameters(RHICmdList, ViewportInfo.UITargetRT->GetRHI(), UITargetRTMaskTexture, ViewportInfo.HDRSourceRT->GetRHI(), ViewportInfo.ColorSpaceLUT);
+							PixelShader->SetParameters(RHICmdList, ViewportInfo.UITargetRT->GetRHI(), UITargetRTMaskTexture, FinalBufferCopy->GetRHI(), ViewportInfo.ColorSpaceLUT);
 						}
 
 						RendererModule.DrawRectangle(
diff --git a/Engine/Source/Runtime/SlateRHIRenderer/Private/SlateRHIRenderer.h b/Engine/Source/Runtime/SlateRHIRenderer/Private/SlateRHIRenderer.h
--- a/Engine/Source/Runtime/SlateRHIRenderer/Private/SlateRHIRenderer.h
+++ b/Engine/Source/Runtime/SlateRHIRenderer/Private/SlateRHIRenderer.h
@@ -63,8 +63,6 @@ struct FViewportInfo : public FRenderResource
 	/** sRGB UI render target */
 	TRefCountPtr<IPooledRenderTarget> UITargetRT;
 	TRefCountPtr<IPooledRenderTarget> UITargetRTMask;
-	/** Copy of the HDR backbuffer */
-	TRefCountPtr<IPooledRenderTarget> HDRSourceRT;
 
 	/** Color-space LUT for HDR UI composition. */
 	FTextureRHIRef ColorSpaceLUT;
diff --git a/Engine/Source/Runtime/SlateRHIRenderer/Private/SlateRHIRenderingPolicy.cpp b/Engine/Source/Runtime/SlateRHIRenderer/Private/SlateRHIRenderingPolicy.cpp
--- a/Engine/Source/Runtime/SlateRHIRenderer/Private/SlateRHIRenderingPolicy.cpp
+++ b/Engine/Source/Runtime/SlateRHIRenderer/Private/SlateRHIRenderingPolicy.cpp
@@ -264,7 +264,7 @@ static bool UpdateScissorRect(
 	uint32& MaskingID,
 	FSlateBackBuffer& BackBuffer,
 	const FSlateRenderBatch& RenderBatch, 
-	FTexture2DRHIRef& ColorTarget,
+	FRHIRenderPassInfo& RPInfo,
 	FTexture2DRHIRef& DepthStencilTarget,
 	const FSlateClippingState*& LastClippingState,
 	const FVector2f ViewTranslation2D, 
@@ -296,7 +296,6 @@ static bool UpdateScissorRect(
 					bDidRestartRenderpass = true;
 					ERenderTargetActions StencilAction = IsMemorylessTexture(DepthStencilTarget) ? ERenderTargetActions::DontLoad_DontStore : ERenderTargetActions::Load_Store;
 
-					FRHIRenderPassInfo RPInfo(ColorTarget, ERenderTargetActions::Load_Store);
 					RPInfo.DepthStencilRenderTarget.Action = MakeDepthStencilTargetActions(ERenderTargetActions::DontLoad_DontStore, StencilAction);
 					RPInfo.DepthStencilRenderTarget.DepthStencilTarget = DepthStencilTarget;
 					RPInfo.DepthStencilRenderTarget.ExclusiveDepthStencil = FExclusiveDepthStencil::DepthNop_StencilWrite;
@@ -388,7 +387,6 @@ static bool UpdateScissorRect(
 						StencilAction = bClearStencil ? ERenderTargetActions::Clear_DontStore : ERenderTargetActions::DontLoad_DontStore;
 					}
 
-					FRHIRenderPassInfo RPInfo(ColorTarget, ERenderTargetActions::Load_Store);
 					RPInfo.DepthStencilRenderTarget.Action = MakeDepthStencilTargetActions(ERenderTargetActions::DontLoad_DontStore, StencilAction);
 					RPInfo.DepthStencilRenderTarget.DepthStencilTarget = DepthStencilTarget;
 					RPInfo.DepthStencilRenderTarget.ExclusiveDepthStencil = FExclusiveDepthStencil::DepthNop_StencilWrite;
@@ -406,7 +404,7 @@ static bool UpdateScissorRect(
 				{
 					FGraphicsPipelineStateInitializer WriteMaskPSOInit;
 					RHICmdList.ApplyCachedRenderTargets(WriteMaskPSOInit);
-					WriteMaskPSOInit.BlendState = TStaticBlendStateWriteMask<CW_NONE>::GetRHI();
+					WriteMaskPSOInit.BlendState = TStaticBlendStateWriteMask<CW_NONE, CW_NONE, CW_NONE, CW_NONE, CW_NONE, CW_NONE, CW_NONE, CW_NONE>::GetRHI();
 					WriteMaskPSOInit.RasterizerState = TStaticRasterizerState<>::GetRHI();
 					WriteMaskPSOInit.DepthStencilState =
 						TStaticDepthStencilState<
@@ -534,6 +532,45 @@ static bool UpdateScissorRect(
 	return bDidRestartRenderpass;
 }
 
+static bool UpdateScissorRect(
+	FRHICommandList& RHICmdList,
+#if STATS
+	int32& ScissorClips,
+	int32& StencilClips,
+#endif
+	uint32& StencilRef,
+	uint32& MaskingID,
+	FSlateBackBuffer& BackBuffer,
+	const FSlateRenderBatch& RenderBatch,
+	FTexture2DRHIRef& ColorTarget,
+	FTexture2DRHIRef& DepthStencilTarget,
+	const FSlateClippingState*& LastClippingState,
+	const FVector2f ViewTranslation2D,
+	FGraphicsPipelineStateInitializer& InGraphicsPSOInit,
+	FSlateStencilClipVertexBuffer& StencilVertexBuffer,
+	const FMatrix& ViewProjection,
+	bool bForceStateChange)
+{
+	FRHIRenderPassInfo RPInfo(ColorTarget, ERenderTargetActions::Load_Store);
+	return UpdateScissorRect(RHICmdList,
+#if STATS
+		ScissorClips,
+		StencilClips,
+#endif
+		StencilRef,
+		MaskingID,
+		BackBuffer,
+		RenderBatch,
+		RPInfo,
+		DepthStencilTarget,
+		LastClippingState,
+		ViewTranslation2D,
+		InGraphicsPSOInit,
+		StencilVertexBuffer,
+		ViewProjection,
+		bForceStateChange);
+}
+
 void FSlateRHIRenderingPolicy::DrawElements(
 	FRHICommandListImmediate& RHICmdList,
 	FSlateBackBuffer& BackBuffer,
@@ -1164,10 +1201,9 @@ void FSlateRHIRenderingPolicy::DrawElements(
 				RectParams.SourceTextureSize = PostProcessTexture->GetSizeXY();
 				RectParams.CornerRadius = ShaderParams.PixelParams3;
 				RectParams.UITarget = Params.UITarget;
-				RectParams.UITargetMask = Params.UITargetMask;
 				RectParams.HDRDisplayColorGamut = Params.HDRDisplayColorGamut;
 
-				RectParams.RestoreStateFunc = [&](FRHICommandListImmediate&InRHICmdList, FGraphicsPipelineStateInitializer& InGraphicsPSOInit) {
+				RectParams.RestoreStateFunc = [&](FRHICommandListImmediate&InRHICmdList, FGraphicsPipelineStateInitializer& InGraphicsPSOInit, FRHIRenderPassInfo& RPInfo) {
 					return UpdateScissorRect(
 						InRHICmdList,
 #if STATS
@@ -1178,7 +1214,7 @@ void FSlateRHIRenderingPolicy::DrawElements(
 						MaskingID,
 						BackBuffer,
 						RenderBatch,
-						ColorTarget,
+						RPInfo,
 						DepthStencilTarget,
 						LastClippingState,
 						ViewTranslation2D,
diff --git a/Engine/Source/Runtime/SlateRHIRenderer/Private/SlateRHIRenderingPolicy.h b/Engine/Source/Runtime/SlateRHIRenderer/Private/SlateRHIRenderingPolicy.h
--- a/Engine/Source/Runtime/SlateRHIRenderer/Private/SlateRHIRenderingPolicy.h
+++ b/Engine/Source/Runtime/SlateRHIRenderer/Private/SlateRHIRenderingPolicy.h
@@ -30,7 +30,6 @@ struct FSlateRenderingParams
 	FVector2f ViewOffset;
 	FGameTime Time;
 	TRefCountPtr<IPooledRenderTarget> UITarget;
-	TRefCountPtr<IPooledRenderTarget> UITargetMask;
 	EDisplayColorGamut HDRDisplayColorGamut;
 	bool bWireFrame;
 	bool bIsHDR;
diff --git a/Engine/Source/Runtime/SlateRHIRenderer/Private/SlateShaders.cpp b/Engine/Source/Runtime/SlateRHIRenderer/Private/SlateShaders.cpp
--- a/Engine/Source/Runtime/SlateRHIRenderer/Private/SlateShaders.cpp
+++ b/Engine/Source/Runtime/SlateRHIRenderer/Private/SlateShaders.cpp
@@ -18,8 +18,9 @@ IMPLEMENT_SHADER_TYPE(, FSlateDebugOverdrawPS, TEXT("/Engine/Private/SlateElemen
 
 IMPLEMENT_SHADER_TYPE(, FSlatePostProcessBlurPS, TEXT("/Engine/Private/SlatePostProcessPixelShader.usf"), TEXT("GaussianBlurMain"), SF_Pixel);
 IMPLEMENT_SHADER_TYPE(, FSlatePostProcessDownsamplePS, TEXT("/Engine/Private/SlatePostProcessPixelShader.usf"), TEXT("DownsampleMain"), SF_Pixel);
-IMPLEMENT_SHADER_TYPE(template<>, FSlatePostProcessUpsamplePS<true>, TEXT("/Engine/Private/SlatePostProcessPixelShader.usf"), TEXT("UpsampleMain"), SF_Pixel);
-IMPLEMENT_SHADER_TYPE(template<>, FSlatePostProcessUpsamplePS<false>, TEXT("/Engine/Private/SlatePostProcessPixelShader.usf"), TEXT("UpsampleMain"), SF_Pixel);
+IMPLEMENT_SHADER_TYPE(template<>, FSlatePostProcessUpsamplePS<ESlatePostProcessUpsamplePSPermutation::SDR>, TEXT("/Engine/Private/SlatePostProcessPixelShader.usf"), TEXT("UpsampleMain"), SF_Pixel);
+IMPLEMENT_SHADER_TYPE(template<>, FSlatePostProcessUpsamplePS<ESlatePostProcessUpsamplePSPermutation::HDR_SCRGB>, TEXT("/Engine/Private/SlatePostProcessPixelShader.usf"), TEXT("UpsampleMain"), SF_Pixel);
+IMPLEMENT_SHADER_TYPE(template<>, FSlatePostProcessUpsamplePS<ESlatePostProcessUpsamplePSPermutation::HDR_PQ10>, TEXT("/Engine/Private/SlatePostProcessPixelShader.usf"), TEXT("UpsampleMain"), SF_Pixel);
 IMPLEMENT_SHADER_TYPE(, FSlatePostProcessColorDeficiencyPS, TEXT("/Engine/Private/SlatePostProcessColorDeficiencyPixelShader.usf"), TEXT("ColorDeficiencyMain"), SF_Pixel);
 
 IMPLEMENT_SHADER_TYPE(, FSlateMaskingVS, TEXT("/Engine/Private/SlateMaskingShader.usf"), TEXT("MainVS"), SF_Vertex);
diff --git a/Engine/Source/Runtime/SlateRHIRenderer/Private/SlateShaders.h b/Engine/Source/Runtime/SlateRHIRenderer/Private/SlateShaders.h
--- a/Engine/Source/Runtime/SlateRHIRenderer/Private/SlateShaders.h
+++ b/Engine/Source/Runtime/SlateRHIRenderer/Private/SlateShaders.h
@@ -394,7 +394,14 @@ private:
 	LAYOUT_FIELD(FShaderParameter, UVBounds);
 };
 
-template<bool bOutputToUITarget>
+enum class ESlatePostProcessUpsamplePSPermutation
+{
+	SDR = 0,
+	HDR_SCRGB,
+	HDR_PQ10,
+};
+
+template<ESlatePostProcessUpsamplePSPermutation SlatePostProcessUpsamplePSPermutation>
 class FSlatePostProcessUpsamplePS : public FSlateElementPS
 {
 	DECLARE_SHADER_TYPE(FSlatePostProcessUpsamplePS, Global);
@@ -421,7 +428,8 @@ public:
 	static void ModifyCompilationEnvironment(const FGlobalShaderPermutationParameters& Parameters, FShaderCompilerEnvironment& OutEnvironment)
 	{
 		// Set defines based on what this shader will be used for
-		OutEnvironment.SetDefine(TEXT("OUTPUT_TO_UI_TARGET"), (uint32)(bOutputToUITarget ? 1 : 0));
+		OutEnvironment.SetDefine(TEXT("OUTPUT_TO_UI_TARGET"), (uint32)(SlatePostProcessUpsamplePSPermutation != ESlatePostProcessUpsamplePSPermutation::SDR ? 1 : 0));
+		OutEnvironment.SetDefine(TEXT("SCRGB_ENCODING"), SlatePostProcessUpsamplePSPermutation == ESlatePostProcessUpsamplePSPermutation::HDR_SCRGB);
 		FSlateElementPS::ModifyCompilationEnvironment(Parameters, OutEnvironment);
 	}
 
