Fix transitions for background blur

diff --git a/Engine/Source/Runtime/SlateRHIRenderer/Private/SlatePostProcessor.cpp b/Engine/Source/Runtime/SlateRHIRenderer/Private/SlatePostProcessor.cpp
--- a/Engine/Source/Runtime/SlateRHIRenderer/Private/SlatePostProcessor.cpp
+++ b/Engine/Source/Runtime/SlateRHIRenderer/Private/SlatePostProcessor.cpp
@@ -126,6 +126,7 @@ static void BlitUIToHDRScene(FRHICommandListImmediate& RHICmdList, IRendererModu
 		IPooledRenderTarget* RenderTargets[] = { RectParams.UITarget.GetReference() };
 		FRenderTargetWriteMask::Decode(RHICmdList, ShaderMap, RenderTargets, UITargetRTMask, TexCreate_None, TEXT("UIRTWriteMask"));
 	}
+	FRHITexture* UITargetRTMaskTexture = RHISupportsRenderTargetWriteMask(GMaxRHIShaderPlatform) ? UITargetRTMask->GetRHI() : nullptr;
 
 	// Source is the viewport.  This is the width and height of the viewport backbuffer
 	const int32 SrcTextureWidth = RectParams.SourceTextureSize.X;
@@ -152,7 +153,20 @@ static void BlitUIToHDRScene(FRHICommandListImmediate& RHICmdList, IRendererModu
 		PixelShader = TShaderMapRef<FBlitUIToHDRPS<0> >(ShaderMap);
 	}
 
-	RHICmdList.Transition(FRHITransitionInfo(DestTexture, ERHIAccess::Unknown, ERHIAccess::RTV));
+	TArray<FRHITransitionInfo> RHITransitionInfos =
+	{
+		FRHITransitionInfo(UITexture, ERHIAccess::Unknown, ERHIAccess::SRVGraphics),
+		FRHITransitionInfo(RectParams.SourceTexture, ERHIAccess::Unknown, ERHIAccess::SRVGraphics),
+		FRHITransitionInfo(DestTexture, ERHIAccess::Unknown, ERHIAccess::RTV)
+	};
+
+
+	if (UITargetRTMaskTexture)
+	{
+		RHITransitionInfos.Add(FRHITransitionInfo(UITargetRTMaskTexture, ERHIAccess::Unknown, ERHIAccess::SRVGraphics));
+	}
+
+	RHICmdList.Transition(RHITransitionInfos);
 
 	const FVector2f InvSrcTextureSize(1.f / SrcTextureWidth, 1.f / SrcTextureHeight);
 
@@ -179,8 +193,6 @@ static void BlitUIToHDRScene(FRHICommandListImmediate& RHICmdList, IRendererModu
 		GraphicsPSOInit.PrimitiveType = PT_TriangleList;
 		SetGraphicsPipelineState(RHICmdList, GraphicsPSOInit, 0);
 
-		FRHITexture* UITargetRTMaskTexture = RHISupportsRenderTargetWriteMask(GMaxRHIShaderPlatform) ? UITargetRTMask->GetRHI() : nullptr;
-
 		PixelShader->SetParameters(RHICmdList, UITexture, RectParams.SourceTexture, UITargetRTMaskTexture, FVector2f(SrcTextureWidth, SrcTextureHeight));
 
 		RendererModule.DrawRectangle(
@@ -587,8 +599,11 @@ void FSlatePostProcessor::UpsampleRect(FRHICommandListImmediate& RHICmdList, IRe
 
 	// Perform Writable transitions first
 
-	RHICmdList.Transition(FRHITransitionInfo(SrcTexture, ERHIAccess::Unknown, ERHIAccess::SRVGraphics));
-	RHICmdList.Transition(FRHITransitionInfo(DestTexture, ERHIAccess::Unknown, ERHIAccess::RTV));
+	TArray<FRHITransitionInfo> RHITransitionInfos =
+	{
+		FRHITransitionInfo(SrcTexture, ERHIAccess::Unknown, ERHIAccess::SRVGraphics),
+		FRHITransitionInfo(DestTexture, ERHIAccess::Unknown, ERHIAccess::RTV)
+	};
 
 	FRHIRenderPassInfo RPInfo(DestTexture, ERenderTargetActions::Load_Store);
 
@@ -604,8 +619,11 @@ void FSlatePostProcessor::UpsampleRect(FRHICommandListImmediate& RHICmdList, IRe
 		RPInfo.ColorRenderTargets[1].Action = ERenderTargetActions::Load_Store;
 		bHasMRT = true;
 		bIsSCRGB = (DestTexture->GetDesc().Format == PF_FloatRGBA);
+		RHITransitionInfos.Add(FRHITransitionInfo(UITargetTexture, ERHIAccess::Unknown, ERHIAccess::RTV));
 	}
 
+	RHICmdList.Transition(RHITransitionInfos);
+
 	RHICmdList.BeginRenderPass(RPInfo, TEXT("UpsampleRect"));
 	{
 		RHICmdList.ApplyCachedRenderTargets(GraphicsPSOInit);
