Add UAV support on swapchain for platforms that support it

Add compute shader support for HDR composition which removes a fullscreen copy

diff --git a/Engine/Shaders/Private/CompositeUIPixelShader.usf b/Engine/Shaders/Private/CompositeUIPixelShader.usf
--- a/Engine/Shaders/Private/CompositeUIPixelShader.usf
+++ b/Engine/Shaders/Private/CompositeUIPixelShader.usf
@@ -185,10 +185,9 @@ float4 ComputeHDRUIColor(float2 InputUV)
 	return UIColor;
 }
 
-float3 ComputeHDRSceneColor(float2 InputUV, SamplerState InSceneSampler)
+float3 ComputeHDRSceneColor(float3 SceneColor)
 {
 	const float3x3 sRGB_2_Rec2020 = mul(XYZ_2_Rec2020_MAT, sRGB_2_XYZ_MAT);
-	float3 SceneColor = Texture2DSample(SceneTexture, InSceneSampler, InputUV).xyz;
 #if !SCRGB_ENCODING
 	SceneColor.xyz = ST2084ToLinear(SceneColor.xyz);
 #else
@@ -214,7 +213,8 @@ void Main(
 
 	float4 UIColor = ComputeHDRUIColor(Input.UV);
 
-	float3 SceneColor = ComputeHDRSceneColor(Input.UV, SceneSampler);
+	float3 SceneColor = Texture2DSample(SceneTexture, SceneSampler, Input.UV).xyz;
+	SceneColor = ComputeHDRSceneColor(SceneColor);
 	OutColor.xyz = ComposeUIAndScene(SceneColor.xyz, UIColor, UILevel);
 
 	OutColor.w = 1.f;
@@ -228,6 +228,51 @@ void Main(
 #endif
 }
 
+#if USE_COMPUTE_FOR_COMPOSITION
+float4 SceneTextureDimensions;
+RWTexture2D<float4> RWSceneTexture;
+
+[numthreads(NUM_THREADS_PER_GROUP, NUM_THREADS_PER_GROUP, 1)]
+void CompositeUICS(uint3 DispatchThreadId : SV_DispatchThreadID)
+{
+	if (any(DispatchThreadId.xy >= (uint2)SceneTextureDimensions.xy))
+	{
+		return;
+	}
+
+	float2 InputPosition = DispatchThreadId.xy + float2(0.5, 0.5);
+#if PLATFORM_SUPPORTS_RENDERTARGET_WRITE_MASK
+	uint CompositeUIMask = DecodeRTWriteMask(DispatchThreadId.xy, UIWriteMaskTexture, 1);
+	BRANCH
+	if (CompositeUIMask == 0)
+	{
+		return;
+	}
+#endif
+
+	float2 InputUV = InputPosition * SceneTextureDimensions.zw;
+
+	float4 UIColor = ComputeHDRUIColor(InputUV);
+
+	float3 SceneColor = RWSceneTexture[DispatchThreadId.xy].rgb;
+	SceneColor = ComputeHDRSceneColor(SceneColor);
+	float4 OutColor = 0;
+	OutColor.xyz = ComposeUIAndScene(SceneColor.xyz, UIColor, UILevel);
+	OutColor.w = 1.f;
+
+#if !SCRGB_ENCODING
+	// Linear -> PQ
+	OutColor.xyz = LinearToST2084(OutColor.xyz);
+#else
+	const float3x3 Rec2020_2_sRGB = mul(XYZ_2_sRGB_MAT, Rec2020_2_XYZ_MAT);
+	OutColor.xyz = mul(Rec2020_2_sRGB, OutColor.xyz / ScRGBScaleFactor);
+#endif
+
+	RWSceneTexture[DispatchThreadId.xy] = float4(OutColor.xyz, 1);
+}
+
+#endif
+
 #if DISPLAY_MAPPING_PS
 float4x4 TextureToOutputGamutMatrix;
 
@@ -275,7 +320,8 @@ void BlitUIToHDRPS(
 	}
 #endif
 
-	float3 SceneColor = ComputeHDRSceneColor(Input.UV, UISampler);
+	float3 SceneColor = Texture2DSample(SceneTexture, UISampler, Input.UV).xyz;
+	SceneColor = ComputeHDRSceneColor(SceneColor);
 	OutColor.xyz = ComposeUIAndScene(SceneColor.xyz, UIColor, UILevel) / ScRGBScaleFactor;
 	OutColor.a = 1;
 }
diff --git a/Engine/Source/Runtime/D3D12RHI/Private/D3D12RHIPrivate.h b/Engine/Source/Runtime/D3D12RHI/Private/D3D12RHIPrivate.h
--- a/Engine/Source/Runtime/D3D12RHI/Private/D3D12RHIPrivate.h
+++ b/Engine/Source/Runtime/D3D12RHI/Private/D3D12RHIPrivate.h
@@ -391,6 +391,7 @@ public:
 	virtual bool RHIGetRenderQueryResult(FRHIRenderQuery* RenderQuery, uint64& OutResult, bool bWait, uint32 GPUIndex = INDEX_NONE) final override;
 	virtual uint32 RHIGetViewportNextPresentGPUIndex(FRHIViewport* Viewport) final override;
 	virtual FTexture2DRHIRef RHIGetViewportBackBuffer(FRHIViewport* Viewport) final override;
+	virtual FUnorderedAccessViewRHIRef RHIGetViewportBackBufferUAV(FRHIViewport* Viewport) final override;
 	virtual void RHIAliasTextureResources(FTextureRHIRef& DestTexture, FTextureRHIRef& SrcTexture) final override;
 	virtual FTextureRHIRef RHICreateAliasedTexture(FTextureRHIRef& SourceTexture) final override;
 	virtual void RHIGetDisplaysInformation(FDisplayInformationArray& OutDisplayInformation) final override;
diff --git a/Engine/Source/Runtime/D3D12RHI/Private/D3D12RenderTarget.cpp b/Engine/Source/Runtime/D3D12RHI/Private/D3D12RenderTarget.cpp
--- a/Engine/Source/Runtime/D3D12RHI/Private/D3D12RenderTarget.cpp
+++ b/Engine/Source/Runtime/D3D12RHI/Private/D3D12RenderTarget.cpp
@@ -1227,11 +1227,13 @@ void FD3D12DynamicRHI::RHIReadSurfaceData(FRHITexture* InRHITexture, FIntRect In
 
 	// Could be back buffer reference texture, so get the correct D3D12 texture here
 	// We know already that it's a FD3D12Texture2D so cast is safe
+#if D3D12_USE_DUMMY_BACKBUFFER
 	if (EnumHasAnyFlags(InRHITexture->GetFlags(), TexCreate_Presentable))
 	{
 		FD3D12BackBufferReferenceTexture2D* BufferBufferReferenceTexture = (FD3D12BackBufferReferenceTexture2D*)InRHITexture;
 		InRHITexture = BufferBufferReferenceTexture->GetBackBufferTexture();
 	}
+#endif
 
 	// Retrieve the base texture
 	FD3D12CommandContext& CommandContext = GetRHIDevice(GPUIndex)->GetDefaultCommandContext();
diff --git a/Engine/Source/Runtime/D3D12RHI/Private/D3D12Texture.h b/Engine/Source/Runtime/D3D12RHI/Private/D3D12Texture.h
--- a/Engine/Source/Runtime/D3D12RHI/Private/D3D12Texture.h
+++ b/Engine/Source/Runtime/D3D12RHI/Private/D3D12Texture.h
@@ -208,11 +208,13 @@ FORCEINLINE FD3D12Texture* GetD3D12TextureFromRHITexture(FRHITexture* Texture)
 	
 	// If it's the dummy backbuffer then swap with actual current RHI backbuffer right now
 	FRHITexture* RHITexture = Texture;
+#if D3D12_USE_DUMMY_BACKBUFFER
 	if (RHITexture && EnumHasAnyFlags(RHITexture->GetFlags(), TexCreate_Presentable))
 	{
 		FD3D12BackBufferReferenceTexture2D* BufferBufferReferenceTexture = (FD3D12BackBufferReferenceTexture2D*)RHITexture;
 		RHITexture = BufferBufferReferenceTexture->GetBackBufferTexture();
 	}
+#endif
 
 	FD3D12Texture* Result((FD3D12Texture*)RHITexture->GetTextureBaseRHI());
 	check(Result);
diff --git a/Engine/Source/Runtime/D3D12RHI/Private/D3D12Viewport.cpp b/Engine/Source/Runtime/D3D12RHI/Private/D3D12Viewport.cpp
--- a/Engine/Source/Runtime/D3D12RHI/Private/D3D12Viewport.cpp
+++ b/Engine/Source/Runtime/D3D12RHI/Private/D3D12Viewport.cpp
@@ -247,11 +247,18 @@ FD3D12Texture* GetSwapChainSurface(FD3D12Device* Parent, EPixelFormat PixelForma
 	D3D12_RESOURCE_DESC BackBufferDesc = BackBufferResource->GetDesc();
 
 	FString Name = FString::Printf(TEXT("BackBuffer%d"), BackBufferIndex);
+
+	ETextureCreateFlags SwapchainTextureCreateFlags = ETextureCreateFlags::RenderTargetable;
+	if (RHISupportsSwapchainUAVs(GetFeatureLevelShaderPlatform(GMaxRHIFeatureLevel)))
+	{
+		SwapchainTextureCreateFlags |= ETextureCreateFlags::UAV;
+	}
+
 	FRHITextureCreateDesc CreateDesc =
 		FRHITextureCreateDesc::Create2D(*Name)
 		.SetExtent(FIntPoint((uint32)BackBufferDesc.Width, BackBufferDesc.Height))
 		.SetFormat(PixelFormat)
-		.SetFlags(ETextureCreateFlags::RenderTargetable)
+		.SetFlags(SwapchainTextureCreateFlags)
 		.SetInitialState(ERHIAccess::Present);
 
 	FD3D12DynamicRHI* DynamicRHI = FD3D12DynamicRHI::GetD3DRHI();
@@ -438,7 +445,9 @@ void FD3D12Viewport::CalculateSwapChainDepth(int32 DefaultSwapChainDepth)
 #endif // WITH_MGPU
 
 	BackBuffers.Empty();
+	BackBuffersUAV.Empty();
 	BackBuffers.AddZeroed(NumBackBuffers);
+	BackBuffersUAV.AddZeroed(NumBackBuffers);
 
 	SDRBackBuffers.Empty();
 	SDRBackBuffers.AddZeroed(NumBackBuffers);
@@ -498,6 +507,24 @@ void FD3D12Viewport::Resize(uint32 InSizeX, uint32 InSizeY, bool bInIsFullscreen
 		BackBuffers[i].SafeRelease();
 		check(BackBuffers[i] == nullptr);
 
+		if (IsValidRef(BackBuffersUAV[i]))
+		{
+			// Tell the back buffer to delete immediately so that we can call resize.
+			if (BackBuffersUAV[i]->GetRefCount() != 1)
+			{
+				UE_LOG(LogD3D12RHI, Log, TEXT("Backbuffer %d leaking with %d refs during Resize."), i, BackBuffersUAV[i]->GetRefCount());
+			}
+			check(BackBuffersUAV[i]->GetRefCount() == 1);
+
+			for (FD3D12UnorderedAccessView& Uav : *BackBuffersUAV[i])
+			{
+				Uav.GetResource()->DoNotDeferDelete();
+			}
+		}
+
+		BackBuffersUAV[i].SafeRelease();
+		check(BackBuffersUAV[i] == nullptr);
+
 		if (IsValidRef(SDRBackBuffers[i]))
 		{
 			if (SDRBackBuffers[i]->GetRefCount() != 1)
@@ -601,7 +628,27 @@ static bool IsCompositionEnabled()
 	return !!bDwmEnabled;
 }
 
-#if WITH_MGPU
+FD3D12Texture* FD3D12Viewport::GetDummyBackBuffer_RenderThread(bool bInIsSDR) const
+{
+#if D3D12_USE_DUMMY_BACKBUFFER
+	return bInIsSDR ? SDRDummyBackBuffer_RenderThread : DummyBackBuffer_RenderThread;
+#else
+	checkNoEntry();
+	return nullptr;
+#endif
+}
+
+FD3D12UnorderedAccessView* FD3D12Viewport::GetBackBufferUAV_RenderThread() const
+{ 
+#if D3D12_USE_DUMMY_BACKBUFFER
+    // See FD3D12Viewport::PresentChecked: if we change fullscreen state (which is detected on RHI thread), we might end up with invalid backbuffer: the safe way is to rely on the dummybackbuffer instead, but 
+	// managing UAV is a bit tricky: disallow UAV on viewports right now
+    checkNoEntry();
+	return nullptr;
+#else
+	return BackBufferUAV_RenderThread;
+#endif
+}
 
 /** Update the expected next present GPU back buffer index from RenderThread point of view */
 void FD3D12Viewport::AdvanceExpectedBackBufferIndex_RenderThread()
@@ -610,11 +657,15 @@ void FD3D12Viewport::AdvanceExpectedBackBufferIndex_RenderThread()
 		CustomPresent->NeedsNativePresent() || CustomPresent->NeedsAdvanceBackbuffer() : true;
 	if (bNeedsNativePresent)
 	{
+#if WITH_MGPU
 		FScopeLock Lock(&ExpectedBackBufferIndexLock);
+#endif
 
 		ExpectedBackBufferIndex_RenderThread++;
 		ExpectedBackBufferIndex_RenderThread = ExpectedBackBufferIndex_RenderThread % NumBackBuffers;
 
+		BackBuffer_RenderThread = BackBuffers[ExpectedBackBufferIndex_RenderThread];
+		BackBufferUAV_RenderThread = BackBuffersUAV[ExpectedBackBufferIndex_RenderThread];
 #if !UE_BUILD_SHIPPING
 		if (RHIConsoleVariables::LogViewportEvents)
 		{
@@ -625,8 +676,6 @@ void FD3D12Viewport::AdvanceExpectedBackBufferIndex_RenderThread()
 	}
 }
 
-#endif // WITH_MGPU
-
 /** Presents the swap chain checking the return result. */
 bool FD3D12Viewport::PresentChecked(int32 SyncInterval)
 {
@@ -1148,13 +1197,11 @@ void FD3D12DynamicRHI::RHIAdvanceFrameForGetViewportBackBuffer(FRHIViewport* Vie
 	// Don't need to do anything on the back because dummy back buffer texture is used to make sure the correct back
 	// buffer index is always used on RHI thread
 
-#if WITH_MGPU
 	// But advance the expected present GPU index so the next call to RHIGetViewportNextPresentGPUIndex returns the expected GPU index for the next present.
 	// Warning: when present fails or is not called on the RHIThread then this might not be in sync but RHI thread will fix up the correct state
 	//          Present doesn't happen so shouldn't matter that the index was wrong then
 	FD3D12Viewport* Viewport = FD3D12DynamicRHI::ResourceCast(ViewportRHI);
 	Viewport->AdvanceExpectedBackBufferIndex_RenderThread();
-#endif // WITH_MGPU
 }
 
 uint32 FD3D12DynamicRHI::RHIGetViewportNextPresentGPUIndex(FRHIViewport* ViewportRHI)
@@ -1177,18 +1224,24 @@ FTexture2DRHIRef FD3D12DynamicRHI::RHIGetViewportBackBuffer(FRHIViewport* Viewpo
 	check(IsInRenderingThread());
 
 	const FD3D12Viewport* const Viewport = FD3D12DynamicRHI::ResourceCast(ViewportRHI);
-	const bool sIsSDR = false;
-	FRHITexture2D* const DummyBackBuffer = Viewport->GetDummyBackBuffer_RenderThread(sIsSDR);
-	
+
+	FRHITexture* SelectedBackBuffer = Viewport->GetBackBuffer_RenderThread();
 #if !UE_BUILD_SHIPPING
 	if (RHIConsoleVariables::LogViewportEvents)
 	{
 		const FString& ThreadName = FThreadManager::GetThreadName(FPlatformTLS::GetCurrentThreadId());
-		UE_LOG(LogD3D12RHI, Log, TEXT("Thread %s: RHIGetViewportBackBuffer (Viewport %#016llx: Dummy BackBuffer %#016llx)"), ThreadName.GetCharArray().GetData(), Viewport, DummyBackBuffer);
+		UE_LOG(LogD3D12RHI, Log, TEXT("Thread %s: RHIGetViewportBackBuffer (Viewport %#016llx: BackBuffer %#016llx)"), ThreadName.GetCharArray().GetData(), Viewport, SelectedBackBuffer);
 	}
 #endif
 
-	return DummyBackBuffer;
+	return SelectedBackBuffer;
+}
+
+FUnorderedAccessViewRHIRef FD3D12DynamicRHI::RHIGetViewportBackBufferUAV(FRHIViewport* ViewportRHI)
+{
+	check(IsInRenderingThread());
+	const FD3D12Viewport* const Viewport = FD3D12DynamicRHI::ResourceCast(ViewportRHI);
+	return Viewport->GetBackBufferUAV_RenderThread();
 }
 
 #if defined(D3D12_WITH_DWMAPI) && D3D12_WITH_DWMAPI
diff --git a/Engine/Source/Runtime/D3D12RHI/Private/Windows/WindowsD3D12Viewport.cpp b/Engine/Source/Runtime/D3D12RHI/Private/Windows/WindowsD3D12Viewport.cpp
--- a/Engine/Source/Runtime/D3D12RHI/Private/Windows/WindowsD3D12Viewport.cpp
+++ b/Engine/Source/Runtime/D3D12RHI/Private/Windows/WindowsD3D12Viewport.cpp
@@ -45,10 +45,12 @@ FD3D12Viewport::FD3D12Viewport(class FD3D12Adapter* InParent, HWND InWindowHandl
 	DummyBackBuffer_RenderThread(nullptr),
 	CurrentBackBufferIndex_RHIThread(0),
 	BackBuffer_RHIThread(nullptr),
+	BackBuffer_RenderThread(nullptr),
+	BackBufferUAV_RenderThread(nullptr),
 #if WITH_MGPU
 	BackbufferMultiGPUBinding(0),
-	ExpectedBackBufferIndex_RenderThread(0),
 #endif //WITH_MGPU
+	ExpectedBackBufferIndex_RenderThread(0),
 	SDRDummyBackBuffer_RenderThread(nullptr),
 	SDRBackBuffer_RHIThread(nullptr),
 	SDRPixelFormat(PF_B8G8R8A8),
@@ -371,6 +373,12 @@ void FD3D12Viewport::ResizeInternal()
 	BackBuffer_RHIThread = BackBuffers[CurrentBackBufferIndex_RHIThread].GetReference();
 	SDRBackBuffer_RHIThread = SDRBackBuffers[CurrentBackBufferIndex_RHIThread].GetReference();
 
+#if !D3D12_USE_DUMMY_BACKBUFFER
+    static_assert(false, "D3D12_USE_DUMMY_BACKBUFFER==0 has not been tested for viewports");
+	ExpectedBackBufferIndex_RenderThread = 0;
+	BackBuffer_RenderThread = BackBuffers[ExpectedBackBufferIndex_RenderThread].GetReference();
+#endif
+
 	// Create dummy back buffer which always reference to the actual RHI thread back buffer - can't be bound directly to D3D12
 	DummyBackBuffer_RenderThread = CreateDummyBackBufferTextures(Adapter, PixelFormat, SizeX, SizeY, false);
 	SDRDummyBackBuffer_RenderThread = (SDRBackBuffer_RHIThread != nullptr) ? CreateDummyBackBufferTextures(Adapter, PixelFormat, SizeX, SizeY, true) : nullptr;
@@ -614,6 +622,14 @@ void FD3D12Viewport::OnResumeRendering()
 void FD3D12Viewport::OnSuspendRendering()
 {}
 
+FD3D12Texture* FD3D12Viewport::GetBackBuffer_RenderThread() const
+{
+	check(IsInRenderingThread());
+	const bool sIsSDR = false;
+	FD3D12Texture* const DummyBackBuffer = GetDummyBackBuffer_RenderThread(sIsSDR);
+	return DummyBackBuffer;
+}
+
 void FD3D12DynamicRHI::RHIGetDisplaysInformation(FDisplayInformationArray& OutDisplayInformation)
 {
 	OutDisplayInformation.Append(DisplayList);
diff --git a/Engine/Source/Runtime/D3D12RHI/Public/D3D12RHI.h b/Engine/Source/Runtime/D3D12RHI/Public/D3D12RHI.h
--- a/Engine/Source/Runtime/D3D12RHI/Public/D3D12RHI.h
+++ b/Engine/Source/Runtime/D3D12RHI/Public/D3D12RHI.h
@@ -64,6 +64,11 @@
 	#include "D3D12RHIPlatformPublic.h"
 #endif
 
+#if PLATFORM_WINDOWS
+#define D3D12_USE_DUMMY_BACKBUFFER 1
+#else
+#define D3D12_USE_DUMMY_BACKBUFFER 0
+#endif
 
 #define FD3D12_TEXTURE_DATA_PITCH_ALIGNMENT D3D12_TEXTURE_DATA_PITCH_ALIGNMENT
 #define FD3D12_DESCRIPTOR_HEAP_TYPE_SAMPLER D3D12_DESCRIPTOR_HEAP_TYPE_SAMPLER
diff --git a/Engine/Source/Runtime/D3D12RHI/Public/D3D12Viewport.h b/Engine/Source/Runtime/D3D12RHI/Public/D3D12Viewport.h
--- a/Engine/Source/Runtime/D3D12RHI/Public/D3D12Viewport.h
+++ b/Engine/Source/Runtime/D3D12RHI/Public/D3D12Viewport.h
@@ -92,11 +92,14 @@ public:
 	// Accessors.
 	FIntPoint GetSizeXY() const { return FIntPoint(SizeX, SizeY); }
 
-	FD3D12Texture* GetDummyBackBuffer_RenderThread(bool bInIsSDR) const { return bInIsSDR ? SDRDummyBackBuffer_RenderThread : DummyBackBuffer_RenderThread; }
+	FD3D12Texture* GetDummyBackBuffer_RenderThread(bool bInIsSDR) const;
+	FD3D12Texture* GetBackBuffer_RenderThread() const;
 
 	FD3D12Texture* GetBackBuffer_RHIThread() const { return BackBuffer_RHIThread; }
 	FD3D12Texture* GetSDRBackBuffer_RHIThread() const { return (PixelFormat == SDRPixelFormat) ? GetBackBuffer_RHIThread() : SDRBackBuffer_RHIThread; }
 
+	FD3D12UnorderedAccessView* GetBackBufferUAV_RenderThread() const;
+
 	virtual void WaitForFrameEventCompletion() override;
 	virtual void IssueFrameEvent() override;
 
@@ -135,9 +138,9 @@ public:
 	/** Query the swap chain's current connected output for HDR support. */
 	bool CurrentOutputSupportsHDR() const;
 
-#if WITH_MGPU
 	/** Advance and get the next present GPU index */
 	void AdvanceExpectedBackBufferIndex_RenderThread();
+#if WITH_MGPU
 	uint32 GetNextPresentGPUIndex() const
 	{
 		FScopeLock Lock(&ExpectedBackBufferIndexLock);
@@ -223,16 +226,21 @@ private:
 #endif // D3D12_VIEWPORT_EXPOSES_SWAP_CHAIN
 
 	TArray<TRefCountPtr<FD3D12Texture>> BackBuffers;
+	TArray<TRefCountPtr<FD3D12UnorderedAccessView>> BackBuffersUAV;
 	uint32 NumBackBuffers;
 
 	TRefCountPtr<FD3D12Texture> DummyBackBuffer_RenderThread; // Dummy back buffer texture which always references the current back buffer on the RHI thread
 	uint32 CurrentBackBufferIndex_RHIThread;
 	FD3D12Texture* BackBuffer_RHIThread;
+	FD3D12Texture* BackBuffer_RenderThread;
+	FD3D12UnorderedAccessView* BackBufferUAV_RenderThread;
 
 #if WITH_MGPU
 	int32 BackbufferMultiGPUBinding; // where INDEX_NONE cycles through the GPU, otherwise the GPU index.
 	mutable FCriticalSection ExpectedBackBufferIndexLock; // Can very rarely be modified on the RHI thread as well if present is skipped
+#endif
 	uint32 ExpectedBackBufferIndex_RenderThread; // Expected back buffer GPU index - used and updated on RenderThread!
+#if WITH_MGPU
 	TArray<uint32> BackBufferGPUIndices;
 #endif // WITH_MGPU
 
diff --git a/Engine/Source/Runtime/RHI/Private/RHI.cpp b/Engine/Source/Runtime/RHI/Private/RHI.cpp
--- a/Engine/Source/Runtime/RHI/Private/RHI.cpp
+++ b/Engine/Source/Runtime/RHI/Private/RHI.cpp
@@ -2551,6 +2551,7 @@ void FGenericDataDrivenShaderPlatformInfo::ParseDataDrivenShaderInfo(const FConf
 	GET_SECTION_SUPPORT_HELPER(bSupportsRealTypes);
 	GET_SECTION_INT_HELPER(EnablesHLSL2021ByDefault);
 	GET_SECTION_BOOL_HELPER(bSupportsSceneDataCompressedTransforms);
+	GET_SECTION_BOOL_HELPER(bSupportsSwapchainUAVs);
 #undef GET_SECTION_BOOL_HELPER
 #undef GET_SECTION_INT_HELPER
 #undef GET_SECTION_SUPPORT_HELPER
diff --git a/Engine/Source/Runtime/RHI/Public/DynamicRHI.h b/Engine/Source/Runtime/RHI/Public/DynamicRHI.h
--- a/Engine/Source/Runtime/RHI/Public/DynamicRHI.h
+++ b/Engine/Source/Runtime/RHI/Public/DynamicRHI.h
@@ -1504,6 +1504,11 @@ FORCEINLINE FTexture2DRHIRef RHIGetViewportBackBuffer(FRHIViewport* Viewport)
 	return GDynamicRHI->RHIGetViewportBackBuffer(Viewport);
 }
 
+FORCEINLINE FUnorderedAccessViewRHIRef RHIGetViewportBackBufferUAV(FRHIViewport* Viewport)
+{
+	return GDynamicRHI->RHIGetViewportBackBufferUAV(Viewport);
+}
+
 FORCEINLINE uint32 RHIGetHTilePlatformConfig(uint32 DepthWidth, uint32 DepthHeight)
 {
 	return GDynamicRHI->RHIGetHTilePlatformConfig(DepthWidth, DepthHeight);
diff --git a/Engine/Source/Runtime/RHI/Public/RHI.h b/Engine/Source/Runtime/RHI/Public/RHI.h
--- a/Engine/Source/Runtime/RHI/Public/RHI.h
+++ b/Engine/Source/Runtime/RHI/Public/RHI.h
@@ -210,6 +210,11 @@ inline bool RHISupportsManualVertexFetch(const FStaticShaderPlatform InShaderPla
 	return FDataDrivenShaderPlatformInfo::GetSupportsManualVertexFetch(InShaderPlatform);
 }
 
+inline bool RHISupportsSwapchainUAVs(const FStaticShaderPlatform Platform)
+{
+	return FDataDrivenShaderPlatformInfo::GetSupportsSwapchainUAVs(Platform);
+}
+
 /** 
  * Returns true if SV_VertexID contains BaseVertexIndex passed to the draw call, false if shaders must manually construct an absolute VertexID.
  */
diff --git a/Engine/Source/Runtime/RHI/Public/RHIDefinitions.h b/Engine/Source/Runtime/RHI/Public/RHIDefinitions.h
--- a/Engine/Source/Runtime/RHI/Public/RHIDefinitions.h
+++ b/Engine/Source/Runtime/RHI/Public/RHIDefinitions.h
@@ -426,6 +426,7 @@ class RHI_API FGenericDataDrivenShaderPlatformInfo
 	uint32 EnablesHLSL2021ByDefault : 2; // 0: disabled, 1: global shaders only, 2: all shaders
 	uint32 bSupportsSceneDataCompressedTransforms : 1;
 	uint32 bIsPreviewPlatform : 1;
+	uint32 bSupportsSwapchainUAVs : 1;
 		
 #if WITH_EDITOR
 	FText FriendlyName;
@@ -459,7 +460,7 @@ public:
 		check(IsValid(Platform));
 		return Infos[Platform].ShaderFormat;
 	}
-	
+
 	static FORCEINLINE_DEBUGGABLE const EShaderPlatform GetPreviewShaderPlatformParent(const FStaticShaderPlatform Platform)
 	{
 		check(IsValid(Platform));
@@ -635,6 +636,12 @@ public:
 		return Infos[Platform].bSupports4ComponentUAVReadWrite;
 	}
 
+	static FORCEINLINE_DEBUGGABLE const bool GetSupportsSwapchainUAVs(const FStaticShaderPlatform Platform)
+	{
+		check(IsValid(Platform));
+		return Infos[Platform].bSupportsSwapchainUAVs;
+	}
+
 	static FORCEINLINE_DEBUGGABLE const bool GetSupportsRenderTargetWriteMask(const FStaticShaderPlatform Platform)
 	{
 		check(IsValid(Platform));
@@ -1009,7 +1016,7 @@ public:
 	{
 		return Infos[Platform].bSupportsOIT;
 	}
-	
+
 	static FORCEINLINE_DEBUGGABLE const bool GetIsPreviewPlatform(const FStaticShaderPlatform Platform)
 	{
 		return Infos[Platform].bIsPreviewPlatform;
diff --git a/Engine/Source/Runtime/RHI/Public/RHIValidation.h b/Engine/Source/Runtime/RHI/Public/RHIValidation.h
--- a/Engine/Source/Runtime/RHI/Public/RHIValidation.h
+++ b/Engine/Source/Runtime/RHI/Public/RHIValidation.h
@@ -968,7 +968,6 @@ public:
 
 	virtual FUnorderedAccessViewRHIRef RHIGetViewportBackBufferUAV(FRHIViewport* ViewportRHI) override final
 	{
-		checkNoEntry(); // @todo fix barrier tracking for this, or remove this function
 		return RHI->RHIGetViewportBackBufferUAV(ViewportRHI);
 	}
 
diff --git a/Engine/Source/Runtime/SlateRHIRenderer/Private/SlateRHIRenderer.cpp b/Engine/Source/Runtime/SlateRHIRenderer/Private/SlateRHIRenderer.cpp
--- a/Engine/Source/Runtime/SlateRHIRenderer/Private/SlateRHIRenderer.cpp
+++ b/Engine/Source/Runtime/SlateRHIRenderer/Private/SlateRHIRenderer.cpp
@@ -641,10 +641,9 @@ private:
 IMPLEMENT_SHADER_TYPE(, FCompositeLUTGenerationPS, TEXT("/Engine/Private/CompositeUIPixelShader.usf"), TEXT("GenerateLUTPS"), SF_Pixel);
 
 // Pixel shader to composite UI over HDR buffer
-template<uint32 EncodingType>
-class FCompositePS : public FGlobalShader
+class FCompositeShaderBase : public FGlobalShader
 {
-	DECLARE_SHADER_TYPE(FCompositePS, Global);
+	DECLARE_TYPE_LAYOUT(FCompositeShaderBase, NonVirtual);
 public:
 
 	static bool ShouldCompilePermutation(const FGlobalShaderPermutationParameters& Parameters)
@@ -652,116 +651,164 @@ public:
 		return SupportsUICompositionRendering(Parameters.Platform);
 	}
 
-	FCompositePS(const ShaderMetaType::CompiledShaderInitializerType& Initializer) :
+	FCompositeShaderBase(const ShaderMetaType::CompiledShaderInitializerType& Initializer) :
 		FGlobalShader(Initializer)
 	{
 		UITexture.Bind(Initializer.ParameterMap, TEXT("UITexture"));
 		UIWriteMaskTexture.Bind(Initializer.ParameterMap, TEXT("UIWriteMaskTexture"));
 		UISampler.Bind(Initializer.ParameterMap, TEXT("UISampler"));
-		SceneTexture.Bind(Initializer.ParameterMap, TEXT("SceneTexture"));
-		SceneSampler.Bind(Initializer.ParameterMap, TEXT("SceneSampler"));
 		ColorSpaceLUT.Bind(Initializer.ParameterMap, TEXT("ColorSpaceLUT"));
 		ColorSpaceLUTSampler.Bind(Initializer.ParameterMap, TEXT("ColorSpaceLUTSampler"));
 		UILevel.Bind(Initializer.ParameterMap, TEXT("UILevel"));
 		OutputDevice.Bind(Initializer.ParameterMap, TEXT("OutputDevice"));
 	}
-	FCompositePS() {}
+	FCompositeShaderBase() = default;
 
-	void SetParameters(FRHICommandList& RHICmdList, FRHITexture* UITextureRHI, FRHITexture* UITextureWriteMaskRHI, FRHITexture* SceneTextureRHI, FRHITexture* ColorSpaceLUTRHI)
+	template<class TRHIShader>
+	void SetParametersBase(FRHICommandList& RHICmdList, TRHIShader* BoundShader, FRHITexture* UITextureRHI, FRHITexture* UITextureWriteMaskRHI, FRHITexture* ColorSpaceLUTRHI)
 	{
 		static const auto CVarOutputDevice = IConsoleManager::Get().FindTConsoleVariableDataInt(TEXT("r.HDR.Display.OutputDevice"));
 
-		SetTextureParameter(RHICmdList, RHICmdList.GetBoundPixelShader(), UITexture, UISampler, TStaticSamplerState<SF_Point>::GetRHI(), UITextureRHI);
-		SetTextureParameter(RHICmdList, RHICmdList.GetBoundPixelShader(), SceneTexture, SceneSampler, TStaticSamplerState<SF_Point>::GetRHI(), SceneTextureRHI);
-		SetTextureParameter(RHICmdList, RHICmdList.GetBoundPixelShader(), ColorSpaceLUT, ColorSpaceLUTSampler, TStaticSamplerState<SF_Bilinear, AM_Clamp, AM_Clamp, AM_Clamp>::GetRHI(), ColorSpaceLUTRHI);
-		SetShaderValue(RHICmdList, RHICmdList.GetBoundPixelShader(), UILevel, CVarUILevel.GetValueOnRenderThread());
-		SetShaderValue(RHICmdList, RHICmdList.GetBoundPixelShader(), OutputDevice, CVarOutputDevice->GetValueOnRenderThread());
+		SetTextureParameter(RHICmdList, BoundShader, UITexture, UISampler, TStaticSamplerState<SF_Point>::GetRHI(), UITextureRHI);
+		SetTextureParameter(RHICmdList, BoundShader, ColorSpaceLUT, ColorSpaceLUTSampler, TStaticSamplerState<SF_Bilinear, AM_Clamp, AM_Clamp, AM_Clamp>::GetRHI(), ColorSpaceLUTRHI);
+		SetShaderValue(RHICmdList, BoundShader, UILevel, CVarUILevel.GetValueOnRenderThread());
+		SetShaderValue(RHICmdList, BoundShader, OutputDevice, CVarOutputDevice->GetValueOnRenderThread());
 		
 		if (UITextureWriteMaskRHI != nullptr)
 		{
-			SetTextureParameter(RHICmdList, RHICmdList.GetBoundPixelShader(), UIWriteMaskTexture, UITextureWriteMaskRHI);
+			SetTextureParameter(RHICmdList, BoundShader, UIWriteMaskTexture, UITextureWriteMaskRHI);
 		}
 	}
 
-	static void ModifyCompilationEnvironment(const FGlobalShaderPermutationParameters& Parameters, FShaderCompilerEnvironment& OutEnvironment)
-	{
-		FGlobalShader::ModifyCompilationEnvironment(Parameters, OutEnvironment);
-		OutEnvironment.SetDefine(TEXT("SCRGB_ENCODING"), EncodingType);
-	}
-
 	static const TCHAR* GetSourceFilename()
 	{
 		return TEXT("/Engine/Private/CompositeUIPixelShader.usf");
 	}
 
-	static const TCHAR* GetFunctionName()
-	{
-		return TEXT("Main");
-	}
-
 private:
 	LAYOUT_FIELD(FShaderResourceParameter, UITexture);
 	LAYOUT_FIELD(FShaderResourceParameter, UIWriteMaskTexture);
 	LAYOUT_FIELD(FShaderResourceParameter, UISampler);
-	LAYOUT_FIELD(FShaderResourceParameter, SceneTexture);
-	LAYOUT_FIELD(FShaderResourceParameter, SceneSampler);
 	LAYOUT_FIELD(FShaderResourceParameter, ColorSpaceLUT);
 	LAYOUT_FIELD(FShaderResourceParameter, ColorSpaceLUTSampler);
 	LAYOUT_FIELD(FShaderParameter, UILevel);
 	LAYOUT_FIELD(FShaderParameter, OutputDevice);
 };
+IMPLEMENT_TYPE_LAYOUT(FCompositeShaderBase);
 
-
-// Pixel shader to convert UI from PQ 2020 to Linear (and 2020 space)
-class FHDRBackBufferConvertPS : public FGlobalShader
+class FCompositePSBase : public FCompositeShaderBase
 {
-	DECLARE_SHADER_TYPE(FHDRBackBufferConvertPS, Global);
+	DECLARE_TYPE_LAYOUT(FCompositePSBase, NonVirtual);
 public:
-
-	static bool ShouldCompilePermutation(const FGlobalShaderPermutationParameters& Parameters)
+	FCompositePSBase(const ShaderMetaType::CompiledShaderInitializerType& Initializer) : FCompositeShaderBase(Initializer)
 	{
-		return IsFeatureLevelSupported(Parameters.Platform, ERHIFeatureLevel::SM5);
+		SceneTexture.Bind(Initializer.ParameterMap, TEXT("SceneTexture"));
+		SceneSampler.Bind(Initializer.ParameterMap, TEXT("SceneSampler"));
 	}
+	FCompositePSBase() = default;
 
-	FHDRBackBufferConvertPS(const ShaderMetaType::CompiledShaderInitializerType& Initializer) :
-		FGlobalShader(Initializer)
+	static const TCHAR* GetFunctionName()
 	{
-		SceneTexture.Bind(Initializer.ParameterMap, TEXT("SceneTexture"));
-		SceneSampler.Bind(Initializer.ParameterMap, TEXT("SceneSampler"));
+		return TEXT("Main");
 	}
-	FHDRBackBufferConvertPS() {}
 
-	void SetParameters(FRHICommandList& RHICmdList, FRHITexture* SceneTextureRHI)
+	void SetParameters(FRHICommandList& RHICmdList, FRHITexture* UITextureRHI, FRHITexture* UITextureWriteMaskRHI, FRHITexture* SceneTextureRHI, FRHITexture* ColorSpaceLUTRHI)
 	{
+		SetParametersBase(RHICmdList, RHICmdList.GetBoundPixelShader(), UITextureRHI, UITextureWriteMaskRHI, ColorSpaceLUTRHI);
 		SetTextureParameter(RHICmdList, RHICmdList.GetBoundPixelShader(), SceneTexture, SceneSampler, TStaticSamplerState<SF_Point>::GetRHI(), SceneTextureRHI);
 	}
 
+private:
+	LAYOUT_FIELD(FShaderResourceParameter, SceneTexture);
+	LAYOUT_FIELD(FShaderResourceParameter, SceneSampler);
+};
+IMPLEMENT_TYPE_LAYOUT(FCompositePSBase);
+
+template<uint32 EncodingType>
+class FCompositePS : public FCompositePSBase
+{
+	DECLARE_SHADER_TYPE(FCompositePS, Global);
+public:
+	FCompositePS(const ShaderMetaType::CompiledShaderInitializerType& Initializer) : FCompositePSBase(Initializer) {}
+	FCompositePS() = default;
+
 	static void ModifyCompilationEnvironment(const FGlobalShaderPermutationParameters& Parameters, FShaderCompilerEnvironment& OutEnvironment)
 	{
 		FGlobalShader::ModifyCompilationEnvironment(Parameters, OutEnvironment);
+		OutEnvironment.SetDefine(TEXT("SCRGB_ENCODING"), EncodingType);
 	}
+};
+IMPLEMENT_SHADER_TYPE(template<>, FCompositePS<0>, FCompositePSBase::GetSourceFilename(), FCompositePSBase::GetFunctionName(), SF_Pixel);
+IMPLEMENT_SHADER_TYPE(template<>, FCompositePS<1>, FCompositePSBase::GetSourceFilename(), FCompositePSBase::GetFunctionName(), SF_Pixel);
 
-	static const TCHAR* GetSourceFilename()
+class FCompositeCSBase : public FCompositeShaderBase
+{
+	DECLARE_TYPE_LAYOUT(FCompositeCSBase, NonVirtual);
+public:
+
+	static const uint32 NUM_THREADS_PER_GROUP = 16;
+
+	static bool IsShaderSupported(const EShaderPlatform ShaderPlatform)
 	{
-		return TEXT("/Engine/Private/CompositeUIPixelShader.usf");
+		return RHISupports4ComponentUAVReadWrite(ShaderPlatform) && RHISupportsSwapchainUAVs(ShaderPlatform);
+	}
+
+	FCompositeCSBase(const ShaderMetaType::CompiledShaderInitializerType& Initializer) : FCompositeShaderBase(Initializer)
+	{
+		RWSceneTexture.Bind(Initializer.ParameterMap, TEXT("RWSceneTexture"));
+		SceneTextureDimensions.Bind(Initializer.ParameterMap, TEXT("SceneTextureDimensions"));
 	}
+	FCompositeCSBase() = default;
 
 	static const TCHAR* GetFunctionName()
 	{
-		return TEXT("HDRBackBufferConvertPS");
+		return TEXT("CompositeUICS");
+	}
+
+	void SetParameters(FRHICommandList& RHICmdList, FRHITexture* UITextureRHI, FRHITexture* UITextureWriteMaskRHI, FRHIUnorderedAccessView* InRWSceneTexture, FRHITexture* ColorSpaceLUTRHI, const FVector4f& InSceneTextureDimensions)
+	{
+		SetParametersBase(RHICmdList, RHICmdList.GetBoundComputeShader(), UITextureRHI, UITextureWriteMaskRHI, ColorSpaceLUTRHI);
+		SetUAVParameter(RHICmdList, RHICmdList.GetBoundComputeShader(), RWSceneTexture, InRWSceneTexture);
+		SetShaderValue(RHICmdList, RHICmdList.GetBoundComputeShader(), SceneTextureDimensions, InSceneTextureDimensions);
 	}
 
 private:
-	LAYOUT_FIELD(FShaderResourceParameter, SceneTexture);
-	LAYOUT_FIELD(FShaderResourceParameter, SceneSampler);
+	LAYOUT_FIELD(FShaderResourceParameter, RWSceneTexture);
+	LAYOUT_FIELD(FShaderParameter, SceneTextureDimensions);
+};
+IMPLEMENT_TYPE_LAYOUT(FCompositeCSBase);
+
+template<uint32 EncodingType>
+class FCompositeCS : public FCompositeCSBase
+{
+	DECLARE_SHADER_TYPE(FCompositeCS, Global);
+public:
+	FCompositeCS(const ShaderMetaType::CompiledShaderInitializerType& Initializer) : FCompositeCSBase(Initializer) {}
+	FCompositeCS() = default;
+
+	static bool ShouldCompilePermutation(const FGlobalShaderPermutationParameters& Parameters)
+	{
+		if (!SupportsUICompositionRendering(Parameters.Platform))
+		{
+			return false;
+		}
+
+		return IsShaderSupported(Parameters.Platform);
+	}
+
+
+	static void ModifyCompilationEnvironment(const FGlobalShaderPermutationParameters& Parameters, FShaderCompilerEnvironment& OutEnvironment)
+	{
+		FGlobalShader::ModifyCompilationEnvironment(Parameters, OutEnvironment);
+		OutEnvironment.SetDefine(TEXT("USE_COMPUTE_FOR_COMPOSITION"), 1);
+		OutEnvironment.SetDefine(TEXT("SCRGB_ENCODING"), EncodingType);
+		OutEnvironment.SetDefine(TEXT("NUM_THREADS_PER_GROUP"), NUM_THREADS_PER_GROUP);
+	}
 };
 
-IMPLEMENT_SHADER_TYPE(, FHDRBackBufferConvertPS, FHDRBackBufferConvertPS::GetSourceFilename(), FHDRBackBufferConvertPS::GetFunctionName(), SF_Pixel);
+IMPLEMENT_SHADER_TYPE(template<>, FCompositeCS<0>, FCompositeCSBase::GetSourceFilename(), FCompositeCSBase::GetFunctionName(), SF_Compute);
+IMPLEMENT_SHADER_TYPE(template<>, FCompositeCS<1>, FCompositeCSBase::GetSourceFilename(), FCompositeCSBase::GetFunctionName(), SF_Compute);
 
-#define SHADER_VARIATION(A) typedef FCompositePS<A> FCompositePS##A; \
-	IMPLEMENT_SHADER_TYPE2(FCompositePS##A, SF_Pixel);
-SHADER_VARIATION(0)  SHADER_VARIATION(1)
-#undef SHADER_VARIATION
 
 int32 SlateWireFrame = 0;
 static FAutoConsoleVariableRef CVarSlateWireframe(TEXT("Slate.ShowWireFrame"), SlateWireFrame, TEXT(""), ECVF_Default);
@@ -1092,16 +1139,21 @@ void FSlateRHIRenderer::DrawWindow_RenderThread(FRHICommandListImmediate& RHICmd
 				}
 
 				// Composition pass
-				{
-					RHICmdList.Transition(FRHITransitionInfo(ViewportInfo.UITargetRT->GetRHI(), ERHIAccess::Unknown, ERHIAccess::SRVMask));
 
-					if (RHISupportsRenderTargetWriteMask(GMaxRHIShaderPlatform))
-					{
-						IPooledRenderTarget* RenderTargets[] = { ViewportInfo.UITargetRT.GetReference() };
-						FRenderTargetWriteMask::Decode(RHICmdList, ShaderMap, RenderTargets, ViewportInfo.UITargetRTMask, TexCreate_None, TEXT("UIRTWriteMask"));
-					}
+				RHICmdList.Transition(FRHITransitionInfo(ViewportInfo.UITargetRT->GetRHI(), ERHIAccess::Unknown, ERHIAccess::SRVMask));
 
+				if (RHISupportsRenderTargetWriteMask(GMaxRHIShaderPlatform))
+				{
+					IPooledRenderTarget* RenderTargets[] = { ViewportInfo.UITargetRT.GetReference() };
+					FRenderTargetWriteMask::Decode(RHICmdList, ShaderMap, RenderTargets, ViewportInfo.UITargetRTMask, TexCreate_None, TEXT("UIRTWriteMask"));
+				}
+				FRHITexture* UITargetRTMaskTexture = ViewportInfo.UITargetRTMask.IsValid() ? ViewportInfo.UITargetRTMask->GetRHI() : nullptr;
 
+				bool bUseScRGB = (HDROutputDevice == (int32)EDisplayOutputFormat::HDR_ACES_1000nit_ScRGB || HDROutputDevice == (int32)EDisplayOutputFormat::HDR_ACES_2000nit_ScRGB);
+
+				FRHIUnorderedAccessView* BackBufferUAV = (FCompositeCSBase::IsShaderSupported(GetFeatureLevelShaderPlatform(GMaxRHIFeatureLevel)) && ViewportRT == nullptr) ? RHIGetViewportBackBufferUAV(ViewportInfo.ViewportRHI) : nullptr;
+				if (BackBufferUAV == nullptr)
+				{
 					TRefCountPtr<IPooledRenderTarget> FinalBufferCopy;
 					{
 						FPooledRenderTargetDesc Desc(FPooledRenderTargetDesc::Create2DDesc(FIntPoint(ViewportWidth, ViewportHeight),
@@ -1131,35 +1183,24 @@ void FSlateRHIRenderer::DrawWindow_RenderThread(FRHICommandListImmediate& RHICmd
 
 						TShaderMapRef<FScreenVS> VertexShader(ShaderMap);
 
-						FRHITexture* UITargetRTMaskTexture = ViewportInfo.UITargetRTMask.IsValid() ? ViewportInfo.UITargetRTMask->GetRHI() : nullptr;
-						if (HDROutputDevice == (int32)EDisplayOutputFormat::HDR_ACES_1000nit_ScRGB || HDROutputDevice == (int32)EDisplayOutputFormat::HDR_ACES_2000nit_ScRGB)
+						TShaderRef<FCompositePSBase> PixelShader;
+						if (bUseScRGB)
 						{
-							// ScRGB encoding
-							TShaderMapRef<FCompositePS<1>> PixelShader(ShaderMap);
-
-							GraphicsPSOInit.BoundShaderState.VertexDeclarationRHI = GFilterVertexDeclaration.VertexDeclarationRHI;
-							GraphicsPSOInit.BoundShaderState.VertexShaderRHI = VertexShader.GetVertexShader();
-							GraphicsPSOInit.BoundShaderState.PixelShaderRHI = PixelShader.GetPixelShader();
-							GraphicsPSOInit.PrimitiveType = PT_TriangleList;
-
-							SetGraphicsPipelineState(RHICmdList, GraphicsPSOInit, 0);
-
-							PixelShader->SetParameters(RHICmdList, ViewportInfo.UITargetRT->GetRHI(), UITargetRTMaskTexture, FinalBufferCopy->GetRHI(), ViewportInfo.ColorSpaceLUT);
+							PixelShader = TShaderMapRef<FCompositePS<1> >(ShaderMap);
 						}
 						else
 						{
-							// ST2084 (PQ) encoding
-							TShaderMapRef<FCompositePS<0>> PixelShader(ShaderMap);
+							PixelShader = TShaderMapRef<FCompositePS<0> >(ShaderMap);
+						}
 
-							GraphicsPSOInit.BoundShaderState.VertexDeclarationRHI = GFilterVertexDeclaration.VertexDeclarationRHI;
-							GraphicsPSOInit.BoundShaderState.VertexShaderRHI = VertexShader.GetVertexShader();
-							GraphicsPSOInit.BoundShaderState.PixelShaderRHI = PixelShader.GetPixelShader();
-							GraphicsPSOInit.PrimitiveType = PT_TriangleList;
+						GraphicsPSOInit.BoundShaderState.VertexDeclarationRHI = GFilterVertexDeclaration.VertexDeclarationRHI;
+						GraphicsPSOInit.BoundShaderState.VertexShaderRHI = VertexShader.GetVertexShader();
+						GraphicsPSOInit.BoundShaderState.PixelShaderRHI = PixelShader.GetPixelShader();
+						GraphicsPSOInit.PrimitiveType = PT_TriangleList;
 
-							SetGraphicsPipelineState(RHICmdList, GraphicsPSOInit, 0);
+						SetGraphicsPipelineState(RHICmdList, GraphicsPSOInit, 0);
 
-							PixelShader->SetParameters(RHICmdList, ViewportInfo.UITargetRT->GetRHI(), UITargetRTMaskTexture, FinalBufferCopy->GetRHI(), ViewportInfo.ColorSpaceLUT);
-						}
+						PixelShader->SetParameters(RHICmdList, ViewportInfo.UITargetRT->GetRHI(), UITargetRTMaskTexture, FinalBufferCopy->GetRHI(), ViewportInfo.ColorSpaceLUT);
 
 						RendererModule.DrawRectangle(
 							RHICmdList,
@@ -1174,6 +1215,29 @@ void FSlateRHIRenderer::DrawWindow_RenderThread(FRHICommandListImmediate& RHICmd
 					}
 					RHICmdList.EndRenderPass();
 				}
+				else
+				{
+					TShaderRef<FCompositeCSBase> ComputeShader;
+					if (bUseScRGB)
+					{
+						ComputeShader = TShaderMapRef<FCompositeCS<1> >(ShaderMap);
+					}
+					else
+					{
+						ComputeShader = TShaderMapRef<FCompositeCS<0> >(ShaderMap);
+					}
+
+					RHICmdList.Transition({
+						FRHITransitionInfo(FinalBuffer, ERHIAccess::Unknown, ERHIAccess::UAVCompute)
+					});
+
+					FVector4f SceneTextureDimensions((float)ViewportWidth, (float)ViewportHeight, 1.0f/(float)ViewportWidth, 1.0f/(float)ViewportHeight);
+					SetComputePipelineState(RHICmdList, ComputeShader.GetComputeShader());
+					ComputeShader->SetParameters(RHICmdList, ViewportInfo.UITargetRT->GetRHI(), UITargetRTMaskTexture, BackBufferUAV, ViewportInfo.ColorSpaceLUT, SceneTextureDimensions);
+
+					RHICmdList.DispatchComputeShader(FMath::DivideAndRoundUp(ViewportWidth, FCompositeCSBase::NUM_THREADS_PER_GROUP), FMath::DivideAndRoundUp(ViewportHeight, FCompositeCSBase::NUM_THREADS_PER_GROUP), 1);
+
+				}
 
 				// Put the backbuffer back to the correct one.
 				BackBuffer = FinalBuffer;
