fix ui luminance not being properly set for Compute Shader composition

diff --git a/Engine/Source/Runtime/SlateRHIRenderer/Private/SlateRHIRenderer.cpp b/Engine/Source/Runtime/SlateRHIRenderer/Private/SlateRHIRenderer.cpp
--- a/Engine/Source/Runtime/SlateRHIRenderer/Private/SlateRHIRenderer.cpp
+++ b/Engine/Source/Runtime/SlateRHIRenderer/Private/SlateRHIRenderer.cpp
@@ -706,7 +706,7 @@ public:
 		SetTextureParameter(RHICmdList, BoundShader, ColorSpaceLUT, ColorSpaceLUTSampler, TStaticSamplerState<SF_Bilinear, AM_Clamp, AM_Clamp, AM_Clamp>::GetRHI(), ColorSpaceLUTRHI);
 		SetShaderValue(RHICmdList, BoundShader, UILevel, CVarUILevel.GetValueOnRenderThread());
 		SetShaderValue(RHICmdList, BoundShader, OutputDevice, CVarOutputDevice->GetValueOnRenderThread());
-		SetShaderValue(RHICmdList, RHICmdList.GetBoundPixelShader(), UILuminance, CVarHDRUILuminance.GetValueOnRenderThread());
+		SetShaderValue(RHICmdList, BoundShader, UILuminance, CVarHDRUILuminance.GetValueOnRenderThread());
 
 		if (UITextureWriteMaskRHI != nullptr)
 		{
