Fix SlatePostProcessor when multiple pixel formats are requested (for example, HDR w/ background blur and color deficiency), which require PF_FloatR11G11B10 and PF_B8G8R8A8 every frame

diff --git a/Engine/Source/Runtime/SlateRHIRenderer/Private/SlatePostProcessResource.h b/Engine/Source/Runtime/SlateRHIRenderer/Private/SlatePostProcessResource.h
--- a/Engine/Source/Runtime/SlateRHIRenderer/Private/SlatePostProcessResource.h
+++ b/Engine/Source/Runtime/SlateRHIRenderer/Private/SlatePostProcessResource.h
@@ -35,7 +35,7 @@ public:
 	virtual uint32 GetHeight() const override { return RenderTargetSize.Y; }
 	virtual ESlateShaderResource::Type GetType() const override { return ESlateShaderResource::PostProcess; }
 
-
+	EPixelFormat GetPixelFormat() const { return PixelFormat; }
 private:
 	/** Resizes targets to the new size */
 	void ResizeTargets(const FIntPoint& NewSize, EPixelFormat RequestedPixelFormat);
diff --git a/Engine/Source/Runtime/SlateRHIRenderer/Private/SlatePostProcessor.cpp b/Engine/Source/Runtime/SlateRHIRenderer/Private/SlatePostProcessor.cpp
--- a/Engine/Source/Runtime/SlateRHIRenderer/Private/SlatePostProcessor.cpp
+++ b/Engine/Source/Runtime/SlateRHIRenderer/Private/SlatePostProcessor.cpp
@@ -15,17 +15,42 @@
 DECLARE_CYCLE_STAT(TEXT("Slate PostProcessing RT"), STAT_SlatePostProcessingRTTime, STATGROUP_Slate);
 DECLARE_CYCLE_STAT(TEXT("Slate ColorDeficiency RT"), STAT_SlateColorDeficiencyRTTime, STATGROUP_Slate);
 
+static const int32 NumIntermediateTargets = 2;
+
+FSlatePostProcessResource* FindSlatePostProcessResource(TArray<FSlatePostProcessResource*>& IntermediateTargetsArray, EPixelFormat PixelFormat)
+{
+	for (int32 ui = 0; ui < IntermediateTargetsArray.Num(); ++ui)
+	{
+		if (IntermediateTargetsArray[ui]->GetPixelFormat() == PixelFormat || IntermediateTargetsArray[ui]->GetPixelFormat() == PF_Unknown)
+		{
+			return IntermediateTargetsArray[ui];
+		}
+	}
+
+	FSlatePostProcessResource* NewSlatePostProcessResource = new FSlatePostProcessResource(NumIntermediateTargets);
+	IntermediateTargetsArray.Add(NewSlatePostProcessResource);
+	BeginInitResource(NewSlatePostProcessResource);
+	return NewSlatePostProcessResource;
+}
+
 FSlatePostProcessor::FSlatePostProcessor()
 {
-	const int32 NumIntermediateTargets = 2;
-	IntermediateTargets = new FSlatePostProcessResource(NumIntermediateTargets);
-	BeginInitResource(IntermediateTargets);
+	uint32 MaximumDifferentPixelFormat = 8;
+ 	for (uint32 ui=0; ui < MaximumDifferentPixelFormat; ++ui)
+ 	{
+ 		IntermediateTargetsArray.Add(new FSlatePostProcessResource(NumIntermediateTargets));
+ 		BeginInitResource(IntermediateTargetsArray[ui]);
+ 	}
 }
 
 FSlatePostProcessor::~FSlatePostProcessor()
 {
 	// Note this is deleted automatically because it implements FDeferredCleanupInterface.
-	IntermediateTargets->CleanUp();
+	for (int32 ui = 0; ui < IntermediateTargetsArray.Num(); ++ui)
+	{
+		IntermediateTargetsArray[ui]->CleanUp();
+	}
+ 	IntermediateTargetsArray.Empty();
 }
 
 
@@ -246,6 +271,7 @@ void FSlatePostProcessor::BlurRect(FRHICommandListImmediate& RHICmdList, IRender
 		IntermediatePixelFormat = PF_FloatR11G11B10;
 	}
 
+	FSlatePostProcessResource* IntermediateTargets = FindSlatePostProcessResource(IntermediateTargetsArray, IntermediatePixelFormat);
 	IntermediateTargets->Update(RequiredSize, IntermediatePixelFormat);
 
 	if (bIsHDRSource)
@@ -256,7 +282,7 @@ void FSlatePostProcessor::BlurRect(FRHICommandListImmediate& RHICmdList, IRender
 
 	else if(bDownsample)
 	{
-		DownsampleRect(RHICmdList, RendererModule, RectParams, DownsampleSize);
+		DownsampleRect(RHICmdList, RendererModule, RectParams, DownsampleSize, IntermediateTargets);
 	}
 
 	FSamplerStateRHIRef BilinearClamp = TStaticSamplerState<SF_Bilinear, AM_Clamp, AM_Clamp, AM_Clamp>::GetRHI();
@@ -398,7 +424,7 @@ void FSlatePostProcessor::BlurRect(FRHICommandListImmediate& RHICmdList, IRender
 
 #endif
 
-	UpsampleRect(RHICmdList, RendererModule, RectParams, DownsampleSize, BilinearClamp);
+	UpsampleRect(RHICmdList, RendererModule, RectParams, DownsampleSize, BilinearClamp, IntermediateTargets);
 }
 
 void FSlatePostProcessor::ColorDeficiency(FRHICommandListImmediate& RHICmdList, IRendererModule& RendererModule, const FPostProcessRectParams& RectParams)
@@ -408,6 +434,8 @@ void FSlatePostProcessor::ColorDeficiency(FRHICommandListImmediate& RHICmdList,
 	FIntPoint DestRectSize = RectParams.DestRect.GetSize().IntPoint();
 	FIntPoint RequiredSize = DestRectSize;
 
+	FSlatePostProcessResource* IntermediateTargets = FindSlatePostProcessResource(IntermediateTargetsArray, RectParams.SourceTexture->GetDesc().Format);
+
 	IntermediateTargets->Update(RequiredSize, RectParams.SourceTexture->GetDesc().Format);
 
 #if 1
@@ -477,7 +505,7 @@ void FSlatePostProcessor::ColorDeficiency(FRHICommandListImmediate& RHICmdList,
 	}
 
 	const FIntPoint DownsampleSize = RequiredSize;
-	UpsampleRect(RHICmdList, RendererModule, RectParams, DownsampleSize, PointClamp);
+	UpsampleRect(RHICmdList, RendererModule, RectParams, DownsampleSize, PointClamp, IntermediateTargets);
 
 #endif
 }
@@ -486,10 +514,13 @@ void FSlatePostProcessor::ReleaseRenderTargets()
 {
 	check(IsInGameThread());
 	// Only release the resource not delete it.  Deleting it could cause issues on any RHI thread
-	BeginReleaseResource(IntermediateTargets);
+	for (int32 ui = 0; ui < IntermediateTargetsArray.Num(); ++ui)
+	{
+		BeginReleaseResource(IntermediateTargetsArray[ui]);
+	}
 }
 
-void FSlatePostProcessor::DownsampleRect(FRHICommandListImmediate& RHICmdList, IRendererModule& RendererModule, const FPostProcessRectParams& Params, const FIntPoint& DownsampleSize)
+void FSlatePostProcessor::DownsampleRect(FRHICommandListImmediate& RHICmdList, IRendererModule& RendererModule, const FPostProcessRectParams& Params, const FIntPoint& DownsampleSize, FSlatePostProcessResource* IntermediateTargets)
 {
 	SCOPED_DRAW_EVENT(RHICmdList, SlatePostProcessDownsample);
 
@@ -565,11 +596,11 @@ void FSlatePostProcessor::DownsampleRect(FRHICommandListImmediate& RHICmdList, I
 	
 	// Testing only
 #if 0
-	UpsampleRect(RHICmdList, RendererModule, Params, DownsampleSize);
+	UpsampleRect(RHICmdList, RendererModule, Params, DownsampleSize, IntermediateTargets);
 #endif
 }
 
-void FSlatePostProcessor::UpsampleRect(FRHICommandListImmediate& RHICmdList, IRendererModule& RendererModule, const FPostProcessRectParams& Params, const FIntPoint& DownsampleSize, FSamplerStateRHIRef& Sampler)
+void FSlatePostProcessor::UpsampleRect(FRHICommandListImmediate& RHICmdList, IRendererModule& RendererModule, const FPostProcessRectParams& Params, const FIntPoint& DownsampleSize, FSamplerStateRHIRef& Sampler, FSlatePostProcessResource* IntermediateTargets)
 {
 	SCOPED_DRAW_EVENT(RHICmdList, SlatePostProcessUpsample);
 
diff --git a/Engine/Source/Runtime/SlateRHIRenderer/Private/SlatePostProcessor.h b/Engine/Source/Runtime/SlateRHIRenderer/Private/SlatePostProcessor.h
--- a/Engine/Source/Runtime/SlateRHIRenderer/Private/SlatePostProcessor.h
+++ b/Engine/Source/Runtime/SlateRHIRenderer/Private/SlatePostProcessor.h
@@ -43,9 +43,9 @@ public:
 	void ReleaseRenderTargets();
 
 private:
-	void DownsampleRect(FRHICommandListImmediate& RHICmdList, IRendererModule& RendererModule, const FPostProcessRectParams& Params, const FIntPoint& DownsampleSize);
-	void UpsampleRect(FRHICommandListImmediate& RHICmdList, IRendererModule& RendererModule, const FPostProcessRectParams& Params, const FIntPoint& DownsampleSize, FSamplerStateRHIRef& Sampler);
+	void DownsampleRect(FRHICommandListImmediate& RHICmdList, IRendererModule& RendererModule, const FPostProcessRectParams& Params, const FIntPoint& DownsampleSize, FSlatePostProcessResource* IntermediateTargets);
+	void UpsampleRect(FRHICommandListImmediate& RHICmdList, IRendererModule& RendererModule, const FPostProcessRectParams& Params, const FIntPoint& DownsampleSize, FSamplerStateRHIRef& Sampler, FSlatePostProcessResource* IntermediateTargets);
 	int32 ComputeBlurWeights(int32 KernelSize, float StdDev, TArray<FVector4f>& OutWeightsAndOffsets);
 private:
-	FSlatePostProcessResource* IntermediateTargets;
+	TArray<FSlatePostProcessResource*> IntermediateTargetsArray;
 };
\ No newline at end of file
