Fix color deficiency filter (CDF) with HDR

The current problem is that the CDF is not taking into account PLATFORM_SUPPORTS_RENDERTARGET_WRITE_MASK and was only applying the filter to the UI rendertarget.

If the fix was applied inside FSlatePostProcessor::ColorDeficiency, we'd have to add support for PLATFORM_SUPPORTS_RENDERTARGET_WRITE_MASK, call FRenderTargetWriteMask::Decode, and add additional copies in order to apply the filter to both the HDR and SDR targets.

Since we apply the color deficiency filter at the end of FSlateRHIRenderingPolicy::DrawElements, we can apply the CDF inside the composition pass in HDR instead. We need to be careful to apply it to the whole screen now instead of relying on CompositeUIMask, but it's still cheaper and simpler than doing all the copies / processing I mentionned above

diff --git a/Engine/Shaders/Private/ColorDeficiency.ush b/Engine/Shaders/Private/ColorDeficiency.ush
new file mode 100644
--- /dev/null
+++ b/Engine/Shaders/Private/ColorDeficiency.ush
@@ -0,0 +1,252 @@
+// CDV tables comes from:
+// A Model for Simulation of Color Vision Deficiency and A Color Contrast Enhancement Technique for Dichromats
+// - Gustavo M. Machado (2010)
+
+static const float3 RGB_2_Protanomaly[11 * 3] =
+{
+	// Severity - 0
+	float3(1.0, 0.0, 0.0),
+	float3(0.0, 1.0, 0.0),
+	float3(0.0, 0.0, 1.0),
+	// Severity - 1
+	float3(0.856167, 0.182038, -0.038205),
+	float3(0.029342, 0.955115, 0.015544),
+	float3(-0.002880, -0.001563, 1.004443),
+	// Severity - 2
+	float3(0.734766, 0.334872, -0.069637),
+	float3(0.051840, 0.919198, 0.028963),
+	float3(-0.004928, -0.004209, 1.009137),
+	// Severity - 3
+	float3(0.630323, 0.465641, -0.095964),
+	float3(0.069181, 0.890046, 0.040773),
+	float3(-0.006308, -0.007724, 1.014032),
+	// Severity - 4
+	float3(0.539009, 0.579343, -0.118352),
+	float3(0.082546, 0.866121, 0.051332),
+	float3(-0.007136, -0.011959, 1.019095),
+	// Severity - 5
+	float3(0.458064, 0.679578, -0.137642),
+	float3(0.092785, 0.846313, 0.060902),
+	float3(-0.007494, -0.016807, 1.024301),
+	// Severity - 6
+	float3(0.385450, 0.769005, -0.154455),
+	float3(0.100526, 0.829802, 0.069673),
+	float3(-0.007442, -0.022190, 1.029632),
+	// Severity - 7
+	float3(0.319627, 0.849633, -0.169261),
+	float3(0.106241, 0.815969, 0.077790),
+	float3(-0.007025, -0.028051, 1.035076),
+	// Severity - 8
+	float3(0.259411, 0.923008, -0.182420),
+	float3(0.110296, 0.804340, 0.085364),
+	float3(-0.006276, -0.034346, 1.040622),
+	// Severity - 9
+	float3(0.203876, 0.990338, -0.194214),
+	float3(0.112975, 0.794542, 0.092483),
+	float3(-0.005222, -0.041043, 1.046265),
+	// Severity - 10
+	float3(0.152286, 1.052583, -0.204868),
+	float3(0.114503, 0.786281, 0.099216),
+	float3(-0.003882, -0.048116, 1.051998)
+};
+
+static const float3 RGB_2_Deuteranomaly[11 * 3] =
+{
+	// Severity - 0
+	float3(1.0, 0.0, 0.0),
+	float3(0.0, 1.0, 0.0),
+	float3(0.0, 0.0, 1.0),
+	// Severity - 1
+	float3(0.866435, 0.177704, -0.044139),
+	float3(0.049567, 0.939063, 0.011370),
+	float3(-0.003453, 0.007233, 0.996220),
+	// Severity - 2
+	float3(0.760729, 0.319078, -0.079807),
+	float3(0.090568, 0.889315, 0.020117),
+	float3(-0.006027, 0.013325, 0.992702),
+	// Severity - 3
+	float3(0.675425, 0.433850, -0.109275),
+	float3(0.125303, 0.847755, 0.026942),
+	float3(-0.007950, 0.018572, 0.989378),
+	// Severity - 4
+	float3(0.605511, 0.528560, -0.134071),
+	float3(0.155318, 0.812366, 0.032316),
+	float3(-0.009376, 0.023176, 0.986200),
+	// Severity - 5
+	float3(0.547494, 0.607765, -0.155259),
+	float3(0.181692, 0.781742, 0.036566),
+	float3(-0.010410, 0.027275, 0.983136),
+	// Severity - 6
+	float3(0.498864, 0.674741, -0.173604),
+	float3(0.205199, 0.754872, 0.039929),
+	float3(-0.011131, 0.030969, 0.980162),
+	// Severity - 7
+	float3(0.457771, 0.731899, -0.189670),
+	float3(0.226409, 0.731012, 0.042579),
+	float3(-0.011595, 0.034333, 0.977261),
+	// Severity - 8
+	float3(0.422823, 0.781057, -0.203881),
+	float3(0.245752, 0.709602, 0.044646),
+	float3(-0.011843, 0.037423, 0.974421),
+	// Severity - 9
+	float3(0.392952, 0.823610, -0.216562),
+	float3(0.263559, 0.690210, 0.046232),
+	float3(-0.011910, 0.040281, 0.971630),
+	// Severity - 10
+	float3(0.367322, 0.860646, -0.227968),
+	float3(0.280085, 0.672501, 0.047413),
+	float3(-0.011820, 0.042940, 0.968881)
+};
+
+static const float3 RGB_2_Tritanomaly[11 * 3] =
+{
+	// Severity - 0
+	float3(1.0, 0.0, 0.0),
+	float3(0.0, 1.0, 0.0),
+	float3(0.0, 0.0, 1.0),
+	// Severity - 1
+	float3(0.926670, 0.092514, -0.019184),
+	float3(0.021191, 0.964503, 0.014306),
+	float3(0.008437, 0.054813, 0.936750),
+	// Severity - 2
+	float3(0.895720, 0.133330, -0.029050),
+	float3(0.029997, 0.945400, 0.024603),
+	float3(0.013027, 0.104707, 0.882266),
+	// Severity - 3
+	float3(0.905871, 0.127791, -0.033662),
+	float3(0.026856, 0.941251, 0.031893),
+	float3(0.013410, 0.148296, 0.838294),
+	// Severity - 4
+	float3(0.948035, 0.089490, -0.037526),
+	float3(0.014364, 0.946792, 0.038844),
+	float3(0.010853, 0.193991, 0.795156),
+	// Severity - 5
+	float3(1.017277, 0.027029, -0.044306),
+	float3(-0.006113, 0.958479, 0.047634),
+	float3(0.006379, 0.248708, 0.744913),
+	// Severity - 6
+	float3(1.104996, -0.046633, -0.058363),
+	float3(-0.032137, 0.971635, 0.060503),
+	float3(0.001336, 0.317922, 0.680742),
+	// Severity - 7
+	float3(1.193214, -0.109812, -0.083402),
+	float3(-0.058496, 0.979410, 0.079086),
+	float3(-0.002346, 0.403492, 0.598854),
+	// Severity - 8
+	float3(1.257728, -0.139648, -0.118081),
+	float3(-0.078003, 0.975409, 0.102594),
+	float3(-0.003316, 0.501214, 0.502102),
+	// Severity - 9
+	float3(1.278864, -0.125333, -0.153531),
+	float3(-0.084748, 0.957674, 0.127074),
+	float3(-0.000989, 0.601151, 0.399838),
+	// Severity - 10
+	float3(1.255528, -0.076749, -0.178779),
+	float3(-0.078411, 0.930809, 0.147602),
+	float3(0.004733, 0.691367, 0.303900)
+};
+
+#define CDV_NormalVision 0
+#define CDV_Deuteranope 1
+#define CDV_Protanope 2
+#define CDV_Tritanope 3
+
+
+float3 ConvertSourceRGBToDeficientRGB(float3 SourceRGB, float ColorVisionDeficiencySeverity, float ColorVisionDeficiencyType)
+{
+	float3 DeficientRGB = SourceRGB;
+	int Index = (int)ColorVisionDeficiencySeverity * 3;
+	if (ColorVisionDeficiencyType == CDV_Deuteranope)
+	{
+		float3x3 Mat = float3x3(RGB_2_Deuteranomaly[Index], RGB_2_Deuteranomaly[Index + 1], RGB_2_Deuteranomaly[Index + 2]);
+		DeficientRGB = mul(Mat, SourceRGB);
+	}
+	else if (ColorVisionDeficiencyType == CDV_Protanope)
+	{
+		float3x3 Mat = float3x3(RGB_2_Protanomaly[Index], RGB_2_Protanomaly[Index + 1], RGB_2_Protanomaly[Index + 2]);
+		DeficientRGB = mul(Mat, SourceRGB);
+	}
+	else if (ColorVisionDeficiencyType == CDV_Tritanope)
+	{
+		float3x3 Mat = float3x3(RGB_2_Tritanomaly[Index], RGB_2_Tritanomaly[Index + 1], RGB_2_Tritanomaly[Index + 2]);
+		DeficientRGB = mul(Mat, SourceRGB);
+	}
+
+	return DeficientRGB;
+}
+
+// http://www.daltonize.org
+static const float3x3 RGB_DALTONIZE_MAT =
+{
+	0.0, 0.7, 0.7,
+	0.0, 1.0, 0.0,
+	0.0, 0.0, 1.0,
+};
+
+static const float3x3 RGB_DALTONIZE_RED_BLIND =
+{
+	0.5, 1.0, 1.0,
+	0.5, 0.5, 0.5,
+	0.5, 0.5, 0.5,
+};
+
+static const float3x3 RGB_DALTONIZE_GREEN_BLIND =
+{
+	0.5, 0.0, 0.5,
+	1.0, 0.0, 1.0,
+	0.5, 0.0, 0.5,
+};
+
+static const float3x3 RGB_DALTONIZE_BLUE_BLIND =
+{
+	0.5, 0.5, 0.5,
+	0.5, 0.5, 0.5,
+	1.0, 1.0, 0.5,
+};
+
+#define USE_RGB_SPACE_CDV_CORRECTION 1
+#define USE_SINGLE_DALTONIZE_MATRIX 0
+
+float3 ColorDeficiency(float3 SourceRGB, float ColorVisionDeficiencyType, float ColorVisionDeficiencySeverity, float bCorrectDeficiency, float bSimulateCorrectionWithDeficiency)
+{
+	float3 DeficientRGB = ConvertSourceRGBToDeficientRGB(SourceRGB, ColorVisionDeficiencySeverity, ColorVisionDeficiencyType);
+
+	float3 OutRGB = DeficientRGB;
+
+	if (bCorrectDeficiency == 1)
+	{
+		#if USE_RGB_SPACE_CDV_CORRECTION
+			// 
+			float3 ErrorRGB = (SourceRGB - DeficientRGB);
+
+			#if USE_SINGLE_DALTONIZE_MATRIX
+				// 
+				float3 CorrectionRGB = mul(RGB_DALTONIZE_MAT, ErrorRGB);
+			#else
+				float3 CorrectionRGB = float3(0.0, 0.0, 0.0);
+				if (ColorVisionDeficiencyType == CDV_Deuteranope)
+				{
+					CorrectionRGB = mul(RGB_DALTONIZE_GREEN_BLIND, ErrorRGB);
+				}
+				else if (ColorVisionDeficiencyType == CDV_Protanope)
+				{
+					CorrectionRGB = mul(RGB_DALTONIZE_RED_BLIND, ErrorRGB);
+				}
+				else if (ColorVisionDeficiencyType == CDV_Tritanope)
+				{
+					CorrectionRGB = mul(RGB_DALTONIZE_BLUE_BLIND, ErrorRGB);
+				}
+			#endif
+
+			// 
+			OutRGB = SourceRGB + CorrectionRGB;
+		#endif
+
+		if (bSimulateCorrectionWithDeficiency == 1)
+		{
+			OutRGB = ConvertSourceRGBToDeficientRGB(OutRGB, ColorVisionDeficiencySeverity, ColorVisionDeficiencyType);
+		}
+	}
+
+	return OutRGB;
+}
diff --git a/Engine/Shaders/Private/CompositeUIPixelShader.usf b/Engine/Shaders/Private/CompositeUIPixelShader.usf
--- a/Engine/Shaders/Private/CompositeUIPixelShader.usf
+++ b/Engine/Shaders/Private/CompositeUIPixelShader.usf
@@ -5,6 +5,7 @@ CompositeUIPixelShader.usf: Filter pixel shader source.
 
 #include "Common.ush"
 #include "TonemapCommon.ush"
+#include "ColorDeficiency.ush"
 
 #define LUT_SCALING	1.05f
 #define LUT_SIZE	32.f
@@ -24,6 +25,13 @@ float4 ACESCoefsHigh_0;
 float ACESCoefsLow_4;
 float ACESCoefsHigh_4;
 
+#if APPLY_COLOR_DEFICIENCY
+float ColorVisionDeficiencyType;
+float ColorVisionDeficiencySeverity;
+float bCorrectDeficiency;
+float bSimulateCorrectionWithDeficiency;
+#endif
+
 void GenerateLUTPS(
 	FWriteToSliceGeometryOutput Input,
 	out float4 OutColor : SV_Target0
@@ -169,6 +177,19 @@ float3 ComposeUIAndScene(float3 SceneColor, float4 InUIColor, float InUILevel)
 	return SceneColor * (1.f - InUIColor.w) + InUIColor.xyz * InUILevel;
 }
 
+float3 ApplyColorDeficiency(float3 Rec2020Color)
+{
+#if APPLY_COLOR_DEFICIENCY
+	const float3x3 Rec2020_2_sRGB = mul(XYZ_2_sRGB_MAT, Rec2020_2_XYZ_MAT);
+	const float3x3 sRGB_2_Rec2020 = mul(XYZ_2_Rec2020_MAT, sRGB_2_XYZ_MAT);
+	float3 SRGBColor = mul(Rec2020_2_sRGB, Rec2020Color.xyz);
+	SRGBColor = ColorDeficiency(SRGBColor, ColorVisionDeficiencyType, ColorVisionDeficiencySeverity, bCorrectDeficiency, bSimulateCorrectionWithDeficiency);
+	Rec2020Color = mul(sRGB_2_Rec2020, SRGBColor);
+#endif
+
+	return Rec2020Color;
+}
+
 // About sRGB_2_Rec2020 in ComputeHDRUIColor/ComputeHDRSceneColor: to get consistent blending between UI and scene color, regardless of the swapchain format choice, we perform blending in the rec2020 color space, and
 // move back to the output device space. Rec2020 is chosen over ScRGB because this is the way blending is currently being done in engine mode, and contains only positive values
 
@@ -205,7 +226,6 @@ float3 ComputeHDRSceneColor(float3 SceneColor)
 	return SceneColor;
 }
 
-
 void Main(
 	FScreenVertexOutput Input,
 	out float4 OutColor : SV_Target0
@@ -213,19 +233,31 @@ void Main(
 {
 #if PLATFORM_SUPPORTS_RENDERTARGET_WRITE_MASK
 	uint CompositeUIMask = DecodeRTWriteMask(uint2(Input.Position.xy), UIWriteMaskTexture, 1);
+	// in case of color deficiency, we will need to update the whole texture, not just the regions that were touched by the UI
+#if (APPLY_COLOR_DEFICIENCY == 0)
 	BRANCH
 	if (CompositeUIMask == 0)
 	{
 		discard;
 	}
+#endif
 #endif
 
 	float4 UIColor = ComputeHDRUIColor(Input.UV);
+#if PLATFORM_SUPPORTS_RENDERTARGET_WRITE_MASK
+	// If the UI texture has a mask of 0, this means the UI color needs to be forced to 0 as it may contain uninitialized data
+	if (CompositeUIMask == 0)
+	{
+		UIColor = 0;
+	}
+#endif
 
 	float3 SceneColor = Texture2DSample(SceneTexture, SceneSampler, Input.UV).xyz;
 	SceneColor = ComputeHDRSceneColor(SceneColor);
 	OutColor.xyz = ComposeUIAndScene(SceneColor.xyz, UIColor, UILevel);
 
+	OutColor.rgb = ApplyColorDeficiency(OutColor.rgb);
+
 	OutColor.w = 1.f;
 
 #if !SCRGB_ENCODING
@@ -252,16 +284,27 @@ void CompositeUICS(uint3 DispatchThreadId : SV_DispatchThreadID)
 	float2 InputPosition = DispatchThreadId.xy + float2(0.5, 0.5);
 #if PLATFORM_SUPPORTS_RENDERTARGET_WRITE_MASK
 	uint CompositeUIMask = DecodeRTWriteMask(DispatchThreadId.xy, UIWriteMaskTexture, 1);
+	// in case of color deficiency, we will need to update the whole texture, not just the regions that were touched by the UI
+#if (APPLY_COLOR_DEFICIENCY == 0)
 	BRANCH
 	if (CompositeUIMask == 0)
 	{
 		return;
 	}
+#endif
 #endif
 
 	float2 InputUV = InputPosition * SceneTextureDimensions.zw;
 
 	float4 UIColor = ComputeHDRUIColor(InputUV);
+#if PLATFORM_SUPPORTS_RENDERTARGET_WRITE_MASK
+	// If the UI texture has a mask of 0, this means the UI color needs to be forced to 0 as it may contain uninitialized data
+	if (CompositeUIMask == 0)
+	{
+		UIColor = 0;
+	}
+#endif
+
 
 	float3 SceneColor = RWSceneTexture[DispatchThreadId.xy].rgb;
 	SceneColor = ComputeHDRSceneColor(SceneColor);
@@ -269,6 +312,8 @@ void CompositeUICS(uint3 DispatchThreadId : SV_DispatchThreadID)
 	OutColor.xyz = ComposeUIAndScene(SceneColor.xyz, UIColor, UILevel);
 	OutColor.w = 1.f;
 
+	OutColor.rgb = ApplyColorDeficiency(OutColor.rgb);
+
 #if !SCRGB_ENCODING
 	// Linear -> PQ
 	OutColor.xyz = LinearToST2084(OutColor.xyz);
diff --git a/Engine/Shaders/Private/SlatePostProcessColorDeficiencyPixelShader.usf b/Engine/Shaders/Private/SlatePostProcessColorDeficiencyPixelShader.usf
--- a/Engine/Shaders/Private/SlatePostProcessColorDeficiencyPixelShader.usf
+++ b/Engine/Shaders/Private/SlatePostProcessColorDeficiencyPixelShader.usf
@@ -4,278 +4,23 @@
 #include "SlateShaderCommon.ush"
 #include "GammaCorrectionCommon.ush"
 #include "ColorSpace.ush"
+#include "ColorDeficiency.ush"
 
 Texture2D ElementTexture;
 SamplerState ElementTextureSampler;
 
-#define CDV_NormalVision 0
-#define CDV_Deuteranope 1
-#define CDV_Protanope 2
-#define CDV_Tritanope 3
-
 float ColorVisionDeficiencyType;
 float ColorVisionDeficiencySeverity;
 float bCorrectDeficiency;
 float bSimulateCorrectionWithDeficiency;
 
-// CDV tables comes from:
-// A Model for Simulation of Color Vision Deficiency and A Color Contrast Enhancement Technique for Dichromats
-// - Gustavo M. Machado (2010)
-
-static const float3 RGB_2_Protanomaly[11 * 3] = 
-{
-	// Severity - 0
-	float3(1.0, 0.0, 0.0),
- 	float3(0.0, 1.0, 0.0),
- 	float3(0.0, 0.0, 1.0),
-	// Severity - 1
-	float3(0.856167, 0.182038, -0.038205),
- 	float3(0.029342, 0.955115, 0.015544),
- 	float3(-0.002880, -0.001563, 1.004443),
-	// Severity - 2
-	float3(0.734766, 0.334872, -0.069637),
- 	float3(0.051840, 0.919198, 0.028963),
- 	float3(-0.004928, -0.004209, 1.009137),
-	// Severity - 3
-	float3(0.630323, 0.465641, -0.095964),
- 	float3(0.069181, 0.890046, 0.040773),
- 	float3(-0.006308, -0.007724, 1.014032),
-	// Severity - 4
-	float3(0.539009, 0.579343, -0.118352),
- 	float3(0.082546, 0.866121, 0.051332),
- 	float3(-0.007136, -0.011959, 1.019095),
-	// Severity - 5
-	float3(0.458064, 0.679578, -0.137642),
- 	float3(0.092785, 0.846313, 0.060902),
- 	float3(-0.007494, -0.016807, 1.024301),
-	// Severity - 6
-	float3(0.385450, 0.769005, -0.154455),
- 	float3(0.100526, 0.829802, 0.069673),
- 	float3(-0.007442, -0.022190, 1.029632),
-	// Severity - 7
-	float3(0.319627, 0.849633, -0.169261),
- 	float3(0.106241, 0.815969, 0.077790),
- 	float3(-0.007025, -0.028051, 1.035076),
-	// Severity - 8
-	float3(0.259411, 0.923008, -0.182420),
- 	float3(0.110296, 0.804340, 0.085364),
- 	float3(-0.006276, -0.034346, 1.040622),
-	// Severity - 9
-	float3(0.203876, 0.990338, -0.194214),
- 	float3(0.112975, 0.794542, 0.092483),
- 	float3(-0.005222, -0.041043, 1.046265),
-	// Severity - 10
-	float3(0.152286, 1.052583, -0.204868),
- 	float3(0.114503, 0.786281, 0.099216),
- 	float3(-0.003882, -0.048116, 1.051998)
-};
-
-static const float3 RGB_2_Deuteranomaly[11 * 3] = 
-{
-	// Severity - 0
-	float3(1.0, 0.0, 0.0),
-	float3(0.0, 1.0, 0.0),
-	float3(0.0, 0.0, 1.0),
-	// Severity - 1
-	float3(0.866435, 0.177704, -0.044139),
-	float3(0.049567, 0.939063, 0.011370),
-	float3(-0.003453, 0.007233, 0.996220),
-	// Severity - 2
-	float3(0.760729, 0.319078, -0.079807),
-	float3(0.090568, 0.889315, 0.020117),
-	float3(-0.006027, 0.013325, 0.992702),
-	// Severity - 3
-	float3(0.675425, 0.433850, -0.109275),
-	float3(0.125303, 0.847755, 0.026942),
-	float3(-0.007950, 0.018572, 0.989378),
-	// Severity - 4
-	float3(0.605511, 0.528560, -0.134071),
-	float3(0.155318, 0.812366, 0.032316),
-	float3(-0.009376, 0.023176, 0.986200),
-	// Severity - 5
-	float3(0.547494, 0.607765, -0.155259),
-	float3(0.181692, 0.781742, 0.036566),
-	float3(-0.010410, 0.027275, 0.983136),
-	// Severity - 6
-	float3(0.498864, 0.674741, -0.173604),
-	float3(0.205199, 0.754872, 0.039929),
-	float3(-0.011131, 0.030969, 0.980162),
-	// Severity - 7
-	float3(0.457771, 0.731899, -0.189670),
-	float3(0.226409, 0.731012, 0.042579),
-	float3(-0.011595, 0.034333, 0.977261),
-	// Severity - 8
-	float3(0.422823, 0.781057, -0.203881),
-	float3(0.245752, 0.709602, 0.044646),
-	float3(-0.011843, 0.037423, 0.974421),
-	// Severity - 9
-	float3(0.392952, 0.823610, -0.216562),
-	float3(0.263559, 0.690210, 0.046232),
-	float3(-0.011910, 0.040281, 0.971630),
-	// Severity - 10
-	float3(0.367322, 0.860646, -0.227968),
-	float3(0.280085, 0.672501, 0.047413),
-	float3(-0.011820, 0.042940, 0.968881)
-};
-
-static const float3 RGB_2_Tritanomaly[11 * 3] = 
-{
-	// Severity - 0
-	float3(1.0, 0.0, 0.0),
-	float3(0.0, 1.0, 0.0),
-	float3(0.0, 0.0, 1.0),
-	// Severity - 1
-	float3(0.926670, 0.092514, -0.019184),
-	float3(0.021191, 0.964503, 0.014306),
-	float3(0.008437, 0.054813, 0.936750),
-	// Severity - 2
-	float3(0.895720, 0.133330, -0.029050),
-	float3(0.029997, 0.945400, 0.024603),
-	float3(0.013027, 0.104707, 0.882266),
-	// Severity - 3
-	float3(0.905871, 0.127791, -0.033662),
-	float3(0.026856, 0.941251, 0.031893),
-	float3(0.013410, 0.148296, 0.838294),
-	// Severity - 4
-	float3(0.948035, 0.089490, -0.037526),
-	float3(0.014364, 0.946792, 0.038844),
-	float3(0.010853, 0.193991, 0.795156),
-	// Severity - 5
-	float3(1.017277, 0.027029, -0.044306),
-	float3(-0.006113, 0.958479, 0.047634),
-	float3(0.006379, 0.248708, 0.744913),
-	// Severity - 6
-	float3(1.104996, -0.046633, -0.058363),
-	float3(-0.032137, 0.971635, 0.060503),
-	float3(0.001336, 0.317922, 0.680742),
-	// Severity - 7
-	float3(1.193214, -0.109812, -0.083402),
-	float3(-0.058496, 0.979410, 0.079086),
-	float3(-0.002346, 0.403492, 0.598854),
-	// Severity - 8
-	float3(1.257728, -0.139648, -0.118081),
-	float3(-0.078003, 0.975409, 0.102594),
-	float3(-0.003316, 0.501214, 0.502102),
-	// Severity - 9
-	float3(1.278864, -0.125333, -0.153531),
-	float3(-0.084748, 0.957674, 0.127074),
-	float3(-0.000989, 0.601151, 0.399838),
-	// Severity - 10
-	float3(1.255528, -0.076749, -0.178779),
-	float3(-0.078411, 0.930809, 0.147602),
-	float3(0.004733, 0.691367, 0.303900)
-};
-
-
-
-
-float3 ConvertSourceRGBToDeficientRGB(float3 SourceRGB)
-{
-	float3 DeficientRGB = SourceRGB;
-	int Index = (int)ColorVisionDeficiencySeverity * 3;
-	if (ColorVisionDeficiencyType == CDV_Deuteranope)
-	{
-		float3x3 Mat = float3x3(RGB_2_Deuteranomaly[Index], RGB_2_Deuteranomaly[Index + 1], RGB_2_Deuteranomaly[Index + 2]);
-		DeficientRGB = mul(Mat, SourceRGB);
-	}
-	else if (ColorVisionDeficiencyType == CDV_Protanope)
-	{
-		float3x3 Mat = float3x3(RGB_2_Protanomaly[Index], RGB_2_Protanomaly[Index + 1], RGB_2_Protanomaly[Index + 2]);
-		DeficientRGB = mul(Mat, SourceRGB);
-	}
-	else if (ColorVisionDeficiencyType == CDV_Tritanope)
-	{
-		float3x3 Mat = float3x3(RGB_2_Tritanomaly[Index], RGB_2_Tritanomaly[Index + 1], RGB_2_Tritanomaly[Index + 2]);
-		DeficientRGB = mul(Mat, SourceRGB);
-	}
-
-	return DeficientRGB;
-}
-
-
-
-
-// http://www.daltonize.org
-static const float3x3 RGB_DALTONIZE_MAT =
-{
-	0.0, 0.7, 0.7,
-	0.0, 1.0, 0.0,
-	0.0, 0.0, 1.0,
-};
-
-static const float3x3 RGB_DALTONIZE_RED_BLIND =
-{
-	0.5, 1.0, 1.0,
-	0.5, 0.5, 0.5,
-	0.5, 0.5, 0.5,
-};
-
-static const float3x3 RGB_DALTONIZE_GREEN_BLIND =
-{
-	0.5, 0.0, 0.5,
-	1.0, 0.0, 1.0,
-	0.5, 0.0, 0.5,
-};
-
-static const float3x3 RGB_DALTONIZE_BLUE_BLIND =
-{
-	0.5, 0.5, 0.5,
-	0.5, 0.5, 0.5,
-	1.0, 1.0, 0.5,
-};
-
-
-#define USE_RGB_SPACE_CDV_CORRECTION 1
-#define USE_SINGLE_DALTONIZE_MATRIX 0
-
-
 float4 ColorDeficiencyMain(FScreenVertexOutput Input) : SV_Target0
 {
 	float4 SourceColor = Texture2DSample(ElementTexture, ElementTextureSampler, Input.UV);
-	float3 SourceRGB = SourceColor.rgb;
-
 	// Have to convert to linear color space first - so they say, but it doesn't appear to matter much.
-	//SourceRGB = sRGBToLinear(SourceRGB);
-
-	float3 DeficientRGB = ConvertSourceRGBToDeficientRGB(SourceRGB);
-
-	float3 OutRGB = DeficientRGB;
-
-	if (bCorrectDeficiency == 1)
-	{
-		#if USE_RGB_SPACE_CDV_CORRECTION
-			// 
-			float3 ErrorRGB = (SourceRGB - DeficientRGB);
-
-			#if USE_SINGLE_DALTONIZE_MATRIX
-				// 
-				float3 CorrectionRGB = mul( RGB_DALTONIZE_MAT, ErrorRGB );
-			#else
-				float3 CorrectionRGB = float3(0.0, 0.0, 0.0);
-				if (ColorVisionDeficiencyType == CDV_Deuteranope)
-				{
-					CorrectionRGB = mul( RGB_DALTONIZE_GREEN_BLIND, ErrorRGB );
-				}
-				else if (ColorVisionDeficiencyType == CDV_Protanope)
-				{
-					CorrectionRGB = mul( RGB_DALTONIZE_RED_BLIND, ErrorRGB );
-				}
-				else if (ColorVisionDeficiencyType == CDV_Tritanope)
-				{
-					CorrectionRGB = mul( RGB_DALTONIZE_BLUE_BLIND, ErrorRGB );
-				}
-			#endif
-
-			// 
-			OutRGB = SourceRGB + CorrectionRGB;
-		#endif
-
-		if (bSimulateCorrectionWithDeficiency == 1)
-		{
-			OutRGB = ConvertSourceRGBToDeficientRGB(OutRGB);
-		}
-	}
+	//SourceColor.rgb = sRGBToLinear(SourceColor.rgb);
 
+	float3 OutRGB = ColorDeficiency(SourceColor.rgb, ColorVisionDeficiencyType, ColorVisionDeficiencySeverity, bCorrectDeficiency, bSimulateCorrectionWithDeficiency);
 	//OutRGB = LinearToSrgb(OutRGB);
 
 	return float4(OutRGB, SourceColor.a);
diff --git a/Engine/Source/Runtime/SlateRHIRenderer/Private/SlateRHIRenderer.cpp b/Engine/Source/Runtime/SlateRHIRenderer/Private/SlateRHIRenderer.cpp
--- a/Engine/Source/Runtime/SlateRHIRenderer/Private/SlateRHIRenderer.cpp
+++ b/Engine/Source/Runtime/SlateRHIRenderer/Private/SlateRHIRenderer.cpp
@@ -694,6 +694,10 @@ public:
 		UILevel.Bind(Initializer.ParameterMap, TEXT("UILevel"));
 		UILuminance.Bind(Initializer.ParameterMap, TEXT("UILuminance"));
 		OutputDevice.Bind(Initializer.ParameterMap, TEXT("OutputDevice"));
+		ColorVisionDeficiencyType.Bind(Initializer.ParameterMap, TEXT("ColorVisionDeficiencyType"));
+		ColorVisionDeficiencySeverity.Bind(Initializer.ParameterMap, TEXT("ColorVisionDeficiencySeverity"));
+		bCorrectDeficiency.Bind(Initializer.ParameterMap, TEXT("bCorrectDeficiency"));
+		bSimulateCorrectionWithDeficiency.Bind(Initializer.ParameterMap, TEXT("bSimulateCorrectionWithDeficiency"));
 	}
 	FCompositeShaderBase() = default;
 
@@ -714,6 +718,15 @@ public:
 		}
 	}
 
+	template<class TRHIShader>
+	void SetColorDeficiencyParamsBase(FRHICommandList& RHICmdList, TRHIShader* BoundShader, bool bCorrect, EColorVisionDeficiency DeficiencyType, int32 Severity, bool bShowCorrectionWithDeficiency)
+	{
+		SetShaderValue(RHICmdList, BoundShader, ColorVisionDeficiencyType, (float)DeficiencyType);
+		SetShaderValue(RHICmdList, BoundShader, ColorVisionDeficiencySeverity, (float)Severity);
+		SetShaderValue(RHICmdList, BoundShader, bCorrectDeficiency, bCorrect ? 1.0f : 0.0f);
+		SetShaderValue(RHICmdList, BoundShader, bSimulateCorrectionWithDeficiency, bShowCorrectionWithDeficiency ? 1.0f : 0.0f);
+	}
+
 	static const TCHAR* GetSourceFilename()
 	{
 		return TEXT("/Engine/Private/CompositeUIPixelShader.usf");
@@ -728,6 +741,10 @@ private:
 	LAYOUT_FIELD(FShaderParameter, UILevel);
 	LAYOUT_FIELD(FShaderParameter, UILuminance);
 	LAYOUT_FIELD(FShaderParameter, OutputDevice);
+	LAYOUT_FIELD(FShaderParameter, ColorVisionDeficiencyType);
+	LAYOUT_FIELD(FShaderParameter, ColorVisionDeficiencySeverity);
+	LAYOUT_FIELD(FShaderParameter, bCorrectDeficiency);
+	LAYOUT_FIELD(FShaderParameter, bSimulateCorrectionWithDeficiency);
 };
 IMPLEMENT_TYPE_LAYOUT(FCompositeShaderBase);
 
@@ -753,13 +770,19 @@ public:
 		SetTextureParameter(RHICmdList, RHICmdList.GetBoundPixelShader(), SceneTexture, SceneSampler, TStaticSamplerState<SF_Point>::GetRHI(), SceneTextureRHI);
 	}
 
+	void SetColorDeficiencyParams(FRHICommandList& RHICmdList, bool bCorrect, EColorVisionDeficiency DeficiencyType, int32 Severity, bool bShowCorrectionWithDeficiency)
+	{
+		SetColorDeficiencyParamsBase(RHICmdList, RHICmdList.GetBoundPixelShader(), bCorrect, DeficiencyType, Severity, bShowCorrectionWithDeficiency);
+	}
+
+
 private:
 	LAYOUT_FIELD(FShaderResourceParameter, SceneTexture);
 	LAYOUT_FIELD(FShaderResourceParameter, SceneSampler);
 };
 IMPLEMENT_TYPE_LAYOUT(FCompositePSBase);
 
-template<uint32 EncodingType>
+template<uint32 EncodingType, bool bApplyColorDeficiency>
 class FCompositePS : public FCompositePSBase
 {
 	DECLARE_SHADER_TYPE(FCompositePS, Global);
@@ -771,10 +794,19 @@ public:
 	{
 		FGlobalShader::ModifyCompilationEnvironment(Parameters, OutEnvironment);
 		OutEnvironment.SetDefine(TEXT("SCRGB_ENCODING"), EncodingType);
+		OutEnvironment.SetDefine(TEXT("APPLY_COLOR_DEFICIENCY"), uint32(bApplyColorDeficiency ? 1 : 0));
 	}
 };
-IMPLEMENT_SHADER_TYPE(template<>, FCompositePS<0>, FCompositePSBase::GetSourceFilename(), FCompositePSBase::GetFunctionName(), SF_Pixel);
-IMPLEMENT_SHADER_TYPE(template<>, FCompositePS<1>, FCompositePSBase::GetSourceFilename(), FCompositePSBase::GetFunctionName(), SF_Pixel);
+
+typedef FCompositePS<0, false> FCompositePS_PQ;
+typedef FCompositePS<0, true> FCompositePS_PQ_CDV;
+typedef FCompositePS<1, false> FCompositePS_SCRGB;
+typedef FCompositePS<1, true> FCompositePS_SCRGB_CDV;
+
+IMPLEMENT_SHADER_TYPE(template<>, FCompositePS_PQ, FCompositePSBase::GetSourceFilename(), FCompositePSBase::GetFunctionName(), SF_Pixel);
+IMPLEMENT_SHADER_TYPE(template<>, FCompositePS_PQ_CDV,  FCompositePSBase::GetSourceFilename(), FCompositePSBase::GetFunctionName(), SF_Pixel);
+IMPLEMENT_SHADER_TYPE(template<>, FCompositePS_SCRGB, FCompositePSBase::GetSourceFilename(), FCompositePSBase::GetFunctionName(), SF_Pixel);
+IMPLEMENT_SHADER_TYPE(template<>, FCompositePS_SCRGB_CDV,  FCompositePSBase::GetSourceFilename(), FCompositePSBase::GetFunctionName(), SF_Pixel);
 
 class FCompositeCSBase : public FCompositeShaderBase
 {
@@ -807,13 +839,18 @@ public:
 		SetShaderValue(RHICmdList, RHICmdList.GetBoundComputeShader(), SceneTextureDimensions, InSceneTextureDimensions);
 	}
 
+	void SetColorDeficiencyParams(FRHICommandList& RHICmdList, bool bCorrect, EColorVisionDeficiency DeficiencyType, int32 Severity, bool bShowCorrectionWithDeficiency)
+	{
+		SetColorDeficiencyParamsBase(RHICmdList, RHICmdList.GetBoundComputeShader(), bCorrect, DeficiencyType, Severity, bShowCorrectionWithDeficiency);
+	}
+
 private:
 	LAYOUT_FIELD(FShaderResourceParameter, RWSceneTexture);
 	LAYOUT_FIELD(FShaderParameter, SceneTextureDimensions);
 };
 IMPLEMENT_TYPE_LAYOUT(FCompositeCSBase);
 
-template<uint32 EncodingType>
+template<uint32 EncodingType, bool bApplyColorDeficiency>
 class FCompositeCS : public FCompositeCSBase
 {
 	DECLARE_SHADER_TYPE(FCompositeCS, Global);
@@ -838,11 +875,19 @@ public:
 		OutEnvironment.SetDefine(TEXT("USE_COMPUTE_FOR_COMPOSITION"), 1);
 		OutEnvironment.SetDefine(TEXT("SCRGB_ENCODING"), EncodingType);
 		OutEnvironment.SetDefine(TEXT("NUM_THREADS_PER_GROUP"), NUM_THREADS_PER_GROUP);
+		OutEnvironment.SetDefine(TEXT("APPLY_COLOR_DEFICIENCY"), uint32(bApplyColorDeficiency ? 1 : 0));
 	}
 };
 
-IMPLEMENT_SHADER_TYPE(template<>, FCompositeCS<0>, FCompositeCSBase::GetSourceFilename(), FCompositeCSBase::GetFunctionName(), SF_Compute);
-IMPLEMENT_SHADER_TYPE(template<>, FCompositeCS<1>, FCompositeCSBase::GetSourceFilename(), FCompositeCSBase::GetFunctionName(), SF_Compute);
+typedef FCompositeCS<0, false> FCompositeCS_PQ;
+typedef FCompositeCS<0, true>  FCompositeCS_PQ_CDV;
+typedef FCompositeCS<1, false> FCompositeCS_SCRGB;
+typedef FCompositeCS<1, true>  FCompositeCS_SCRGB_CDV;
+
+IMPLEMENT_SHADER_TYPE(template<>, FCompositeCS_PQ,        FCompositeCSBase::GetSourceFilename(), FCompositeCSBase::GetFunctionName(), SF_Compute);
+IMPLEMENT_SHADER_TYPE(template<>, FCompositeCS_PQ_CDV,    FCompositeCSBase::GetSourceFilename(), FCompositeCSBase::GetFunctionName(), SF_Compute);
+IMPLEMENT_SHADER_TYPE(template<>, FCompositeCS_SCRGB,     FCompositeCSBase::GetSourceFilename(), FCompositeCSBase::GetFunctionName(), SF_Compute);
+IMPLEMENT_SHADER_TYPE(template<>, FCompositeCS_SCRGB_CDV, FCompositeCSBase::GetSourceFilename(), FCompositeCSBase::GetFunctionName(), SF_Compute);
 
 
 int32 SlateWireFrame = 0;
@@ -1097,8 +1142,11 @@ void FSlateRHIRenderer::DrawWindow_RenderThread(FRHICommandListImmediate& RHICmd
 			RHICmdList.BeginDrawingViewport(ViewportInfo.ViewportRHI, FTextureRHIRef());
 			RHICmdList.SetViewport(0.f, 0.f, 0.f, (float)ViewportWidth, (float)ViewportHeight, 0.0f);
 
+			bool bApplyColorDeficiencyCorrectionToRestore = RenderingPolicy->GetApplyColorDeficiencyCorrection();
+
 			if (bCompositeUI)
 			{
+				RenderingPolicy->SetApplyColorDeficiencyCorrection(false);
 				FSlateBatchData& BatchDataHDR = WindowElementList.GetBatchDataHDR();
 				bool bHasBatches = BatchDataHDR.GetRenderBatches().Num() > 0;
 				if (bHasBatches)
@@ -1124,6 +1172,8 @@ void FSlateRHIRenderer::DrawWindow_RenderThread(FRHICommandListImmediate& RHICmd
 
 			if (bCompositeUI)
 			{
+				RenderingPolicy->SetApplyColorDeficiencyCorrection(bApplyColorDeficiencyCorrectionToRestore);
+
 				SCOPED_DRAW_EVENT(RHICmdList, SlateUI_Composition);
 
 				static const FName RendererModuleName("Renderer");
@@ -1185,6 +1235,7 @@ void FSlateRHIRenderer::DrawWindow_RenderThread(FRHICommandListImmediate& RHICmd
 				FRHITexture* UITargetRTMaskTexture = ViewportInfo.UITargetRTMask.IsValid() ? ViewportInfo.UITargetRTMask->GetRHI() : nullptr;
 
 				bool bUseScRGB = (HDROutputDevice == (int32)EDisplayOutputFormat::HDR_ACES_1000nit_ScRGB || HDROutputDevice == (int32)EDisplayOutputFormat::HDR_ACES_2000nit_ScRGB);
+				bool bApplyColorDeficiencyFilter = GSlateColorDeficiencyType != EColorVisionDeficiency::NormalVision && GSlateColorDeficiencySeverity > 0;
 
 				FRHIUnorderedAccessView* BackBufferUAV = (FCompositeCSBase::IsShaderSupported(GetFeatureLevelShaderPlatform(GMaxRHIFeatureLevel)) && ViewportRT == nullptr) ? RHIGetViewportBackBufferUAV(ViewportInfo.ViewportRHI) : nullptr;
 				if (BackBufferUAV == nullptr)
@@ -1221,11 +1272,25 @@ void FSlateRHIRenderer::DrawWindow_RenderThread(FRHICommandListImmediate& RHICmd
 						TShaderRef<FCompositePSBase> PixelShader;
 						if (bUseScRGB)
 						{
-							PixelShader = TShaderMapRef<FCompositePS<1> >(ShaderMap);
+							if (bApplyColorDeficiencyFilter)
+							{
+								PixelShader = TShaderMapRef<FCompositePS_SCRGB_CDV>(ShaderMap);
+							}
+							else
+							{
+								PixelShader = TShaderMapRef<FCompositePS_SCRGB>(ShaderMap);
+							}
 						}
 						else
 						{
-							PixelShader = TShaderMapRef<FCompositePS<0> >(ShaderMap);
+							if (bApplyColorDeficiencyFilter)
+							{
+								PixelShader = TShaderMapRef<FCompositePS_PQ_CDV>(ShaderMap);
+							}
+							else
+							{
+								PixelShader = TShaderMapRef<FCompositePS_PQ>(ShaderMap);
+							}
 						}
 
 						GraphicsPSOInit.BoundShaderState.VertexDeclarationRHI = GFilterVertexDeclaration.VertexDeclarationRHI;
@@ -1236,6 +1301,7 @@ void FSlateRHIRenderer::DrawWindow_RenderThread(FRHICommandListImmediate& RHICmd
 						SetGraphicsPipelineState(RHICmdList, GraphicsPSOInit, 0);
 
 						PixelShader->SetParameters(RHICmdList, ViewportInfo.UITargetRT->GetRHI(), UITargetRTMaskTexture, FinalBufferCopy->GetRHI(), ViewportInfo.ColorSpaceLUT);
+						PixelShader->SetColorDeficiencyParams(RHICmdList, GSlateColorDeficiencyCorrection, GSlateColorDeficiencyType, GSlateColorDeficiencySeverity, GSlateShowColorDeficiencyCorrectionWithDeficiency);
 
 						RendererModule.DrawRectangle(
 							RHICmdList,
@@ -1255,11 +1321,25 @@ void FSlateRHIRenderer::DrawWindow_RenderThread(FRHICommandListImmediate& RHICmd
 					TShaderRef<FCompositeCSBase> ComputeShader;
 					if (bUseScRGB)
 					{
-						ComputeShader = TShaderMapRef<FCompositeCS<1> >(ShaderMap);
+						if (bApplyColorDeficiencyFilter)
+						{
+							ComputeShader = TShaderMapRef<FCompositeCS_SCRGB_CDV>(ShaderMap);
+						}
+						else
+						{
+							ComputeShader = TShaderMapRef<FCompositeCS_SCRGB>(ShaderMap);
+						}
 					}
 					else
 					{
-						ComputeShader = TShaderMapRef<FCompositeCS<0> >(ShaderMap);
+						if (bApplyColorDeficiencyFilter)
+						{
+							ComputeShader = TShaderMapRef<FCompositeCS_PQ_CDV>(ShaderMap);
+						}
+						else
+						{
+							ComputeShader = TShaderMapRef<FCompositeCS_PQ>(ShaderMap);
+						}
 					}
 
 					RHICmdList.Transition({
@@ -1269,6 +1349,7 @@ void FSlateRHIRenderer::DrawWindow_RenderThread(FRHICommandListImmediate& RHICmd
 					FVector4f SceneTextureDimensions((float)ViewportWidth, (float)ViewportHeight, 1.0f/(float)ViewportWidth, 1.0f/(float)ViewportHeight);
 					SetComputePipelineState(RHICmdList, ComputeShader.GetComputeShader());
 					ComputeShader->SetParameters(RHICmdList, ViewportInfo.UITargetRT->GetRHI(), UITargetRTMaskTexture, BackBufferUAV, ViewportInfo.ColorSpaceLUT, SceneTextureDimensions);
+					ComputeShader->SetColorDeficiencyParams(RHICmdList, GSlateColorDeficiencyCorrection, GSlateColorDeficiencyType, GSlateColorDeficiencySeverity, GSlateShowColorDeficiencyCorrectionWithDeficiency);
 
 					RHICmdList.DispatchComputeShader(FMath::DivideAndRoundUp(ViewportWidth, FCompositeCSBase::NUM_THREADS_PER_GROUP), FMath::DivideAndRoundUp(ViewportHeight, FCompositeCSBase::NUM_THREADS_PER_GROUP), 1);
 
diff --git a/Engine/Source/Runtime/SlateRHIRenderer/Private/SlateRHIRenderingPolicy.h b/Engine/Source/Runtime/SlateRHIRenderer/Private/SlateRHIRenderingPolicy.h
--- a/Engine/Source/Runtime/SlateRHIRenderer/Private/SlateRHIRenderingPolicy.h
+++ b/Engine/Source/Runtime/SlateRHIRenderer/Private/SlateRHIRenderingPolicy.h
@@ -67,6 +67,8 @@ public:
 	void SetUseGammaCorrection( bool bInUseGammaCorrection ) { bGammaCorrect = bInUseGammaCorrection; }
 	void SetApplyColorDeficiencyCorrection(bool bInApplyColorCorrection) { bApplyColorDeficiencyCorrection = bInApplyColorCorrection; }
 
+	bool GetApplyColorDeficiencyCorrection() const { return bApplyColorDeficiencyCorrection; }
+
 	virtual void AddSceneAt(FSceneInterface* Scene, int32 Index) override;
 	virtual void ClearScenes() override;
 
