Custom scene texture cannot be read in widget material.

diff --git a/Engine/Source/Runtime/SlateRHIRenderer/Private/SlateMaterialShader.cpp b/Engine/Source/Runtime/SlateRHIRenderer/Private/SlateMaterialShader.cpp
--- a/Engine/Source/Runtime/SlateRHIRenderer/Private/SlateMaterialShader.cpp
+++ b/Engine/Source/Runtime/SlateRHIRenderer/Private/SlateMaterialShader.cpp
@@ -20,7 +20,6 @@ void FSlateMaterialShaderVS::ModifyCompilationEnvironment(const FMaterialShaderP
 	OutEnvironment.SetDefine( TEXT("USE_MATERIALS"), 1 );
 	OutEnvironment.SetDefine( TEXT("NUM_CUSTOMIZED_UVS"), Parameters.MaterialParameters.NumCustomizedUVs );
 	OutEnvironment.SetDefine(TEXT("HAS_SCREEN_POSITION"), (bool)Parameters.MaterialParameters.bHasVertexPositionOffsetConnected);
-	OutEnvironment.SetDefine(TEXT("SCENE_TEXTURES_DISABLED"), 1);
 
 	FMaterialShader::ModifyCompilationEnvironment( Parameters, OutEnvironment );
 }
@@ -53,7 +52,6 @@ void FSlateMaterialShaderPS::ModifyCompilationEnvironment(const FMaterialShaderP
 	// Set defines based on what this shader will be used for
 	OutEnvironment.SetDefine( TEXT("USE_MATERIALS"), 1 );
 	OutEnvironment.SetDefine( TEXT("NUM_CUSTOMIZED_UVS"), Parameters.MaterialParameters.NumCustomizedUVs);
-	OutEnvironment.SetDefine( TEXT("SCENE_TEXTURES_DISABLED"), 1);
 
 	FMaterialShader::ModifyCompilationEnvironment( Parameters, OutEnvironment );
 }
diff --git a/Engine/Source/Runtime/SlateRHIRenderer/Private/SlateMaterialShader.h b/Engine/Source/Runtime/SlateRHIRenderer/Private/SlateMaterialShader.h
--- a/Engine/Source/Runtime/SlateRHIRenderer/Private/SlateMaterialShader.h
+++ b/Engine/Source/Runtime/SlateRHIRenderer/Private/SlateMaterialShader.h
@@ -95,7 +95,6 @@ public:
 		FSlateMaterialShaderVS::ModifyCompilationEnvironment(Parameters, OutEnvironment);
 
 		OutEnvironment.SetDefine(TEXT("USE_SLATE_INSTANCING"), (uint32)( bUseInstancing ? 1 : 0 ));
-		OutEnvironment.SetDefine(TEXT("SCENE_TEXTURES_DISABLED"), 1);
 	}
 };
 
@@ -123,6 +122,5 @@ public:
 		FSlateMaterialShaderPS::ModifyCompilationEnvironment(Parameters,OutEnvironment);
 
 		OutEnvironment.SetDefine(TEXT("SHADER_TYPE"), (uint32)ShaderType);
-		OutEnvironment.SetDefine(TEXT("SCENE_TEXTURES_DISABLED"), 1);
 	}
 };
diff --git a/Engine/Source/Runtime/SlateRHIRenderer/Private/SlateRHIRenderingPolicy.cpp b/Engine/Source/Runtime/SlateRHIRenderer/Private/SlateRHIRenderingPolicy.cpp
--- a/Engine/Source/Runtime/SlateRHIRenderer/Private/SlateRHIRenderingPolicy.cpp
+++ b/Engine/Source/Runtime/SlateRHIRenderer/Private/SlateRHIRenderingPolicy.cpp
@@ -234,6 +234,8 @@ static FSceneView* CreateSceneView( FSceneViewFamilyContext* ViewFamilyContext,
 	// TODO LWC
 	ViewUniformShaderParameters.RelativeWorldViewOrigin = (FVector3f)View->ViewMatrices.GetViewOrigin();
 	
+	// Slate materials need this scale to be positive, otherwise it can fail in querying scene textures (e.g., custom stencil)
+	ViewUniformShaderParameters.BufferToSceneTextureScale = FVector2f(1.0f, 1.0f);
 
 	ERHIFeatureLevel::Type RHIFeatureLevel = View->GetFeatureLevel();
 
