Default slate shader to treat BC4 textures as alpha mask
#rnx

diff --git a/Engine/Shaders/Private/SlateElementPixelShader.usf b/Engine/Shaders/Private/SlateElementPixelShader.usf
--- a/Engine/Shaders/Private/SlateElementPixelShader.usf
+++ b/Engine/Shaders/Private/SlateElementPixelShader.usf
@@ -150,16 +150,28 @@ half4 GetTextureColor( VertexToPixelInterpolants InVertex, float2 UV )
 {
 	half4 BaseColor;
 
-#if SAMPLE_VIRTUAL_TEXTURE
-		VTPageTableResult PageTableResult = TextureLoadVirtualPageTableLevel(InPageTableTexture, VTPageTableUniform_Unpack(VTPackedPageTableUniform[0], VTPackedPageTableUniform[1]), UV, VTADDRESSMODE_WRAP, VTADDRESSMODE_WRAP, VTShaderParams.x); 
-		BaseColor = TextureVirtualSampleLevel(ElementTexture, ElementTextureSampler, PageTableResult, (uint)VTShaderParams.y, VTUniform_Unpack(VTPackedUniform));
-#else
-	BaseColor = Texture2DSample(ElementTexture, ElementTextureSampler, UV );
-#endif
+	#if SAMPLE_VIRTUAL_TEXTURE
+			VTPageTableResult PageTableResult = TextureLoadVirtualPageTableLevel(InPageTableTexture, VTPageTableUniform_Unpack(VTPackedPageTableUniform[0], VTPackedPageTableUniform[1]), UV, VTADDRESSMODE_WRAP, VTADDRESSMODE_WRAP, VTShaderParams.x); 
+			BaseColor = TextureVirtualSampleLevel(ElementTexture, ElementTextureSampler, PageTableResult, (uint)VTShaderParams.y, VTUniform_Unpack(VTPackedUniform));
+	#else
 
-#if !USE_TEXTURE_ALPHA
-	BaseColor.a = 1.0f;
-#endif
+		BaseColor = Texture2DSample(ElementTexture, ElementTextureSampler, UV );
+
+		#if USE_TEXTURE_GRAYSCALE
+			#if !USE_TEXTURE_ALPHA
+				BaseColor.rgb = BaseColor.rrr;
+				BaseColor.a = 1;
+			#else
+				BaseColor.a = BaseColor.r;
+				BaseColor.rgb = half3(1,1,1);
+			#endif
+		#endif
+
+	#endif
+
+	#if !USE_TEXTURE_ALPHA
+		BaseColor.a = 1.0f;
+	#endif
 
 	return BaseColor;
 }
diff --git a/Engine/Source/Runtime/SlateRHIRenderer/Private/SlateRHIRenderingPolicy.cpp b/Engine/Source/Runtime/SlateRHIRenderer/Private/SlateRHIRenderingPolicy.cpp
--- a/Engine/Source/Runtime/SlateRHIRenderer/Private/SlateRHIRenderingPolicy.cpp
+++ b/Engine/Source/Runtime/SlateRHIRenderer/Private/SlateRHIRenderingPolicy.cpp
@@ -816,6 +816,10 @@ void FSlateRHIRenderingPolicy::DrawElements(
 				{
 					bool bIsVirtualTexture = false;
 
+					// check if texture is using BC4 compression and set shader to render grayscale
+					bool bUseTextureGrayscale = false;
+
+
 					if ((ShaderResource != nullptr) && (ResourceType == ESlateShaderResource::TextureObject))
 					{
 						FSlateBaseUTextureResource* TextureObjectResource = const_cast<FSlateBaseUTextureResource*>(static_cast<const FSlateBaseUTextureResource*>(ShaderResource));
@@ -823,10 +827,15 @@ void FSlateRHIRenderingPolicy::DrawElements(
 						if (UTexture* TextureObj = TextureObjectResource->GetTextureObject())
 						{
 							bIsVirtualTexture = TextureObj->IsCurrentlyVirtualTextured();
+
+							if (TextureObj->CompressionSettings == TC_Alpha)
+							{
+								bUseTextureGrayscale = true;
+							}
 						}
 					}
 
-					PixelShader = GetTexturePixelShader(ShaderMap, ShaderType, DrawEffects, bIsVirtualTexture);
+					PixelShader = GetTexturePixelShader(ShaderMap, ShaderType, DrawEffects, bUseTextureGrayscale, bIsVirtualTexture);
 				}
 
 #if WITH_SLATE_VISUALIZERS
@@ -1341,7 +1350,7 @@ ETextureSamplerFilter FSlateRHIRenderingPolicy::GetSamplerFilter(const UTexture*
 	return Filter;
 }
 
-TShaderRef<FSlateElementPS> FSlateRHIRenderingPolicy::GetTexturePixelShader( FGlobalShaderMap* ShaderMap, ESlateShader ShaderType, ESlateDrawEffect DrawEffects, bool bIsVirtualTexture )
+TShaderRef<FSlateElementPS> FSlateRHIRenderingPolicy::GetTexturePixelShader(FGlobalShaderMap* ShaderMap, ESlateShader ShaderType, ESlateDrawEffect DrawEffects, bool bUseTextureGrayscale, bool bIsVirtualTexture)
 {
 	TShaderRef<FSlateElementPS> PixelShader;
 
@@ -1366,22 +1375,50 @@ TShaderRef<FSlateElementPS> FSlateRHIRenderingPolicy::GetTexturePixelShader( FGl
 			{
 				if ( bIsVirtualTexture )
 				{
-					PixelShader = TShaderMapRef<TSlateElementPS<ESlateShader::Default, true, true, true> >(ShaderMap);
+					if (bUseTextureGrayscale)
+					{
+						PixelShader = TShaderMapRef<TSlateElementPS<ESlateShader::Default, true, true, true, true> >(ShaderMap);
+					}
+					else
+					{
+						PixelShader = TShaderMapRef<TSlateElementPS<ESlateShader::Default, true, true, false, true> >(ShaderMap);
+					}
 				}
 				else
 				{
-					PixelShader = TShaderMapRef<TSlateElementPS<ESlateShader::Default, true, true, false> >(ShaderMap);
+					if (bUseTextureGrayscale)
+					{
+						PixelShader = TShaderMapRef<TSlateElementPS<ESlateShader::Default, true, true, true, false> >(ShaderMap);
+					}
+					else
+					{
+						PixelShader = TShaderMapRef<TSlateElementPS<ESlateShader::Default, true, true, false, false> >(ShaderMap);
+					}
 				}
 			}
 			else
 			{
 				if ( bIsVirtualTexture )
 				{
-					PixelShader = TShaderMapRef<TSlateElementPS<ESlateShader::Default, true, false, true> >(ShaderMap);
+					if (bUseTextureGrayscale)
+					{
+						PixelShader = TShaderMapRef<TSlateElementPS<ESlateShader::Default, true, false, true, true> >(ShaderMap);
+					}
+					else
+					{
+						PixelShader = TShaderMapRef<TSlateElementPS<ESlateShader::Default, true, false, false, true> >(ShaderMap);
+					}
 				}
 				else
 				{
-					PixelShader = TShaderMapRef<TSlateElementPS<ESlateShader::Default, true, false, false> >(ShaderMap);
+					if (bUseTextureGrayscale)
+					{
+						PixelShader = TShaderMapRef<TSlateElementPS<ESlateShader::Default, true, false, true, false> >(ShaderMap);
+					}
+					else
+					{
+						PixelShader = TShaderMapRef<TSlateElementPS<ESlateShader::Default, true, false, false, false> >(ShaderMap);
+					}
 				}
 			}
 			break;
@@ -1419,22 +1456,50 @@ TShaderRef<FSlateElementPS> FSlateRHIRenderingPolicy::GetTexturePixelShader( FGl
 			{
 				if ( bIsVirtualTexture )
 				{
-					PixelShader = TShaderMapRef<TSlateElementPS<ESlateShader::Default, false, true, true> >(ShaderMap);
+					if (bUseTextureGrayscale)
+					{
+						PixelShader = TShaderMapRef<TSlateElementPS<ESlateShader::Default, false, true, true, true> >(ShaderMap);
+					}
+					else
+					{
+						PixelShader = TShaderMapRef<TSlateElementPS<ESlateShader::Default, false, true, false, true> >(ShaderMap);
+					}
 				}
 				else
 				{
-					PixelShader = TShaderMapRef<TSlateElementPS<ESlateShader::Default, false, true, false> >(ShaderMap);
+					if (bUseTextureGrayscale)
+					{
+						PixelShader = TShaderMapRef<TSlateElementPS<ESlateShader::Default, false, true, true, false> >(ShaderMap);
+					}
+					else
+					{
+						PixelShader = TShaderMapRef<TSlateElementPS<ESlateShader::Default, false, true, false, false> >(ShaderMap);
+					}
 				}
 			}
 			else
 			{
 				if ( bIsVirtualTexture )
 				{
-					PixelShader = TShaderMapRef<TSlateElementPS<ESlateShader::Default, false, false, true> >(ShaderMap);
+					if (bUseTextureGrayscale)
+					{
+						PixelShader = TShaderMapRef<TSlateElementPS<ESlateShader::Default, false, false, true, true> >(ShaderMap);
+					}
+					else
+					{
+						PixelShader = TShaderMapRef<TSlateElementPS<ESlateShader::Default, false, false, false, true> >(ShaderMap);
+					}
 				}
 				else
 				{
-					PixelShader = TShaderMapRef<TSlateElementPS<ESlateShader::Default, false, false, false> >(ShaderMap);
+					if (bUseTextureGrayscale)
+					{
+						PixelShader = TShaderMapRef<TSlateElementPS<ESlateShader::Default, false, false, true, false> >(ShaderMap);
+					}
+					else
+					{
+						PixelShader = TShaderMapRef<TSlateElementPS<ESlateShader::Default, false, false, false, false> >(ShaderMap);
+					}
 				}
 			}
 			break;
diff --git a/Engine/Source/Runtime/SlateRHIRenderer/Private/SlateRHIRenderingPolicy.h b/Engine/Source/Runtime/SlateRHIRenderer/Private/SlateRHIRenderingPolicy.h
--- a/Engine/Source/Runtime/SlateRHIRenderer/Private/SlateRHIRenderingPolicy.h
+++ b/Engine/Source/Runtime/SlateRHIRenderer/Private/SlateRHIRenderingPolicy.h
@@ -87,7 +87,7 @@ private:
 	 * @param DrawEffects	Draw effects being used
 	 * @return The pixel shader for use with the shader type and draw effects
 	 */
-	TShaderRef<FSlateElementPS> GetTexturePixelShader(FGlobalShaderMap* ShaderMap, ESlateShader ShaderType, ESlateDrawEffect DrawEffects, bool bIsVirtualTexture);
+	TShaderRef<FSlateElementPS> GetTexturePixelShader(FGlobalShaderMap* ShaderMap, ESlateShader ShaderType, ESlateDrawEffect DrawEffects, bool bUseTextureGrayscale,  bool bIsVirtualTexture);
 	void ChooseMaterialShaderTypes(ESlateShader ShaderType, bool bUseInstancing, FMaterialShaderTypes& OutShaderTypes);
 
 	/** @return The RHI primitive type from the Slate primitive type */
diff --git a/Engine/Source/Runtime/SlateRHIRenderer/Private/SlateShaders.cpp b/Engine/Source/Runtime/SlateRHIRenderer/Private/SlateShaders.cpp
--- a/Engine/Source/Runtime/SlateRHIRenderer/Private/SlateShaders.cpp
+++ b/Engine/Source/Runtime/SlateRHIRenderer/Private/SlateShaders.cpp
@@ -28,9 +28,9 @@ IMPLEMENT_SHADER_TYPE(, FSlateMaskingPS, TEXT("/Engine/Private/SlateMaskingShade
 
 IMPLEMENT_SHADER_TYPE(, FSlateDebugBatchingPS, TEXT("/Engine/Private/SlateElementPixelShader.usf"), TEXT("DebugBatchingMain"), SF_Pixel );
 
-#define IMPLEMENT_SLATE_PIXELSHADER_TYPE(ShaderType, bDrawDisabledEffect, bUseTextureAlpha, bIsVirtualTexture) \
-	typedef TSlateElementPS<ESlateShader::ShaderType,bDrawDisabledEffect,bUseTextureAlpha,bIsVirtualTexture> TSlateElementPS##ShaderType##bDrawDisabledEffect##bUseTextureAlpha##bIsVirtualTexture##A; \
-	IMPLEMENT_SHADER_TYPE(template<>,TSlateElementPS##ShaderType##bDrawDisabledEffect##bUseTextureAlpha##bIsVirtualTexture##A,TEXT("/Engine/Private/SlateElementPixelShader.usf"),TEXT("Main"),SF_Pixel);
+#define IMPLEMENT_SLATE_PIXELSHADER_TYPE(ShaderType, bDrawDisabledEffect, bUseTextureAlpha, bUseTextureGrayscale, bIsVirtualTexture) \
+	typedef TSlateElementPS<ESlateShader::ShaderType,bDrawDisabledEffect,bUseTextureAlpha, bUseTextureGrayscale, bIsVirtualTexture> TSlateElementPS##ShaderType##bDrawDisabledEffect##bUseTextureAlpha##bUseTextureGrayscale##bIsVirtualTexture##A; \
+	IMPLEMENT_SHADER_TYPE(template<>,TSlateElementPS##ShaderType##bDrawDisabledEffect##bUseTextureAlpha##bUseTextureGrayscale##bIsVirtualTexture##A,TEXT("/Engine/Private/SlateElementPixelShader.usf"),TEXT("Main"),SF_Pixel);
 
 #if WITH_EDITOR
 IMPLEMENT_SHADER_TYPE(, FHDREditorConvertPS, TEXT("/Engine/Private/CompositeUIPixelShader.usf"), TEXT("HDREditorConvert"), SF_Pixel);
@@ -38,30 +38,40 @@ IMPLEMENT_SHADER_TYPE(, FHDREditorConvertPS, TEXT("/Engine/Private/CompositeUIPi
 /**
 * All the different permutations of shaders used by slate. Uses #defines to avoid dynamic branches
 */
-IMPLEMENT_SLATE_PIXELSHADER_TYPE(Default, false, true, false);
-IMPLEMENT_SLATE_PIXELSHADER_TYPE(Default, false, true, true);
-IMPLEMENT_SLATE_PIXELSHADER_TYPE(Border, false, true, false);
-IMPLEMENT_SLATE_PIXELSHADER_TYPE(Default, true, true, false);
-IMPLEMENT_SLATE_PIXELSHADER_TYPE(Default, true, true, true);
-IMPLEMENT_SLATE_PIXELSHADER_TYPE(Border, true, true, false);
-IMPLEMENT_SLATE_PIXELSHADER_TYPE(Default, false, false, false);
-IMPLEMENT_SLATE_PIXELSHADER_TYPE(Default, false, false, true);
-IMPLEMENT_SLATE_PIXELSHADER_TYPE(Border, false, false, false);
-IMPLEMENT_SLATE_PIXELSHADER_TYPE(Default, true, false, false);
-IMPLEMENT_SLATE_PIXELSHADER_TYPE(Default, true, false, true);
-IMPLEMENT_SLATE_PIXELSHADER_TYPE(Border, true, false, false);
-
-IMPLEMENT_SLATE_PIXELSHADER_TYPE(GrayscaleFont, false, true, false);
-IMPLEMENT_SLATE_PIXELSHADER_TYPE(GrayscaleFont, true, true, false);
-
-IMPLEMENT_SLATE_PIXELSHADER_TYPE(ColorFont, false, true, false);
-IMPLEMENT_SLATE_PIXELSHADER_TYPE(ColorFont, true, true, false);
-
-IMPLEMENT_SLATE_PIXELSHADER_TYPE(LineSegment, false, true, false);
-IMPLEMENT_SLATE_PIXELSHADER_TYPE(LineSegment, true, true, false);
-
-IMPLEMENT_SLATE_PIXELSHADER_TYPE(RoundedBox, false, true, false);
-IMPLEMENT_SLATE_PIXELSHADER_TYPE(RoundedBox, true, true, false);
+IMPLEMENT_SLATE_PIXELSHADER_TYPE(Default, false, true, true, false);
+IMPLEMENT_SLATE_PIXELSHADER_TYPE(Default, false, true, true, true);
+IMPLEMENT_SLATE_PIXELSHADER_TYPE(Border, false, true, false, false);
+IMPLEMENT_SLATE_PIXELSHADER_TYPE(Default, true, true, true, false);
+IMPLEMENT_SLATE_PIXELSHADER_TYPE(Default, true, true, true, true);
+IMPLEMENT_SLATE_PIXELSHADER_TYPE(Border, true, true, false, false);
+IMPLEMENT_SLATE_PIXELSHADER_TYPE(Default, false, false, true, false);
+IMPLEMENT_SLATE_PIXELSHADER_TYPE(Default, false, false, true, true);
+IMPLEMENT_SLATE_PIXELSHADER_TYPE(Border, false, false, false, false);
+IMPLEMENT_SLATE_PIXELSHADER_TYPE(Default, true, false, true, false);
+IMPLEMENT_SLATE_PIXELSHADER_TYPE(Default, true, false, true, true);
+IMPLEMENT_SLATE_PIXELSHADER_TYPE(Border, true, false, false, false);
+
+
+IMPLEMENT_SLATE_PIXELSHADER_TYPE(Default, true, true, false, false);
+IMPLEMENT_SLATE_PIXELSHADER_TYPE(Default, true, true, false, true);
+IMPLEMENT_SLATE_PIXELSHADER_TYPE(Default, true, false, false, false);
+IMPLEMENT_SLATE_PIXELSHADER_TYPE(Default, true, false, false, true);
+IMPLEMENT_SLATE_PIXELSHADER_TYPE(Default, false, true, false, false);
+IMPLEMENT_SLATE_PIXELSHADER_TYPE(Default, false, true, false, true);
+IMPLEMENT_SLATE_PIXELSHADER_TYPE(Default, false, false, false, false);
+IMPLEMENT_SLATE_PIXELSHADER_TYPE(Default, false, false, false, true);
+
+IMPLEMENT_SLATE_PIXELSHADER_TYPE(GrayscaleFont, false, true, false, false);
+IMPLEMENT_SLATE_PIXELSHADER_TYPE(GrayscaleFont, true, true, false, false);
+
+IMPLEMENT_SLATE_PIXELSHADER_TYPE(ColorFont, false, true, false, false);
+IMPLEMENT_SLATE_PIXELSHADER_TYPE(ColorFont, true, true, false, false);
+
+IMPLEMENT_SLATE_PIXELSHADER_TYPE(LineSegment, false, true, false, false);
+IMPLEMENT_SLATE_PIXELSHADER_TYPE(LineSegment, true, true, false, false);
+
+IMPLEMENT_SLATE_PIXELSHADER_TYPE(RoundedBox, false, true, false, false);
+IMPLEMENT_SLATE_PIXELSHADER_TYPE(RoundedBox, true, true, false, false);
 
 /** The Slate vertex declaration. */
 TGlobalResource<FSlateVertexDeclaration> GSlateVertexDeclaration;
diff --git a/Engine/Source/Runtime/SlateRHIRenderer/Private/SlateShaders.h b/Engine/Source/Runtime/SlateRHIRenderer/Private/SlateShaders.h
--- a/Engine/Source/Runtime/SlateRHIRenderer/Private/SlateShaders.h
+++ b/Engine/Source/Runtime/SlateRHIRenderer/Private/SlateShaders.h
@@ -218,7 +218,7 @@ private:
 /** 
  * Pixel shader types for all elements
  */
-template<ESlateShader ShaderType, bool bDrawDisabledEffect, bool bUseTextureAlpha=true, bool bIsVirtualTexture=false>
+template<ESlateShader ShaderType, bool bDrawDisabledEffect, bool bUseTextureAlpha=true, bool bUseTextureGrayscale=false, bool bIsVirtualTexture=false>
 class TSlateElementPS : public FSlateElementPS
 {
 	DECLARE_SHADER_TYPE( TSlateElementPS, Global );
@@ -245,8 +245,9 @@ public:
 		OutEnvironment.SetDefine(TEXT("DRAW_DISABLED_EFFECT"), (uint32)( bDrawDisabledEffect ? 1 : 0 ));
 		OutEnvironment.SetDefine(TEXT("USE_TEXTURE_ALPHA"), (uint32)( bUseTextureAlpha ? 1 : 0 ));
 		OutEnvironment.SetDefine(TEXT("USE_MATERIALS"), (uint32)0);
+		OutEnvironment.SetDefine(TEXT("USE_TEXTURE_GRAYSCALE"), (uint32)(bUseTextureGrayscale ? 1 : 0));
 		OutEnvironment.SetDefine(TEXT("SAMPLE_VIRTUAL_TEXTURE"), (uint32)(bIsVirtualTexture ? 1 : 0));
-
+		
 		FSlateElementPS::ModifyCompilationEnvironment( Parameters, OutEnvironment );
 	}
 };
