Fixed an issue introduced by 22425912 where OnVirtualDesktopSizeChanged would happen and FSlateRHIRenderer::DrawWindows_Private is being called from FPreLoadScreenSlateSynchMechanism::RunMainLoop_SlateThread (ie while main thread is busy initializing modules).
The update can wait for after this initialization, and can be skipped if the project does not support HDR in the first place.

diff --git a/Engine/Source/Runtime/SlateRHIRenderer/Private/SlateRHIRenderer.cpp b/Engine/Source/Runtime/SlateRHIRenderer/Private/SlateRHIRenderer.cpp
--- a/Engine/Source/Runtime/SlateRHIRenderer/Private/SlateRHIRenderer.cpp
+++ b/Engine/Source/Runtime/SlateRHIRenderer/Private/SlateRHIRenderer.cpp
@@ -1594,7 +1594,7 @@ void FSlateRHIRenderer::DrawWindows_Private(FSlateDrawBuffer& WindowDrawBuffer)
 {
 	checkSlow(IsThreadSafeForSlateRendering());
 
-	if (bUpdateHDRDisplayInformation)
+	if (bUpdateHDRDisplayInformation && IsHDRAllowed() && IsInGameThread())
 	{
 		FlushRenderingCommands();
 		RHIHandleDisplayChange();
