Fix a crash on mobile OpenGL with a retainer box

diff --git a/Engine/Source/Runtime/SlateRHIRenderer/Private/SlateRHIRenderer.cpp b/Engine/Source/Runtime/SlateRHIRenderer/Private/SlateRHIRenderer.cpp
--- a/Engine/Source/Runtime/SlateRHIRenderer/Private/SlateRHIRenderer.cpp
+++ b/Engine/Source/Runtime/SlateRHIRenderer/Private/SlateRHIRenderer.cpp
@@ -1757,6 +1757,7 @@ void FSlateRHIRenderer::DrawWindows_Private(FSlateDrawBuffer& WindowDrawBuffer)
 				for (const FRenderThreadUpdateContext& Context : Contexts)
 				{
 					Context.Renderer->DrawWindowToTarget_RenderThread(RHICmdList, Context);
+					Context.Renderer->ReleaseDrawBuffer(*Context.WindowDrawBuffer);
 				}
 			}
 		);
diff --git a/Engine/Source/Runtime/SlateRHIRenderer/Public/Interfaces/ISlate3DRenderer.h b/Engine/Source/Runtime/SlateRHIRenderer/Public/Interfaces/ISlate3DRenderer.h
--- a/Engine/Source/Runtime/SlateRHIRenderer/Public/Interfaces/ISlate3DRenderer.h
+++ b/Engine/Source/Runtime/SlateRHIRenderer/Public/Interfaces/ISlate3DRenderer.h
@@ -17,14 +17,18 @@ public:
 	/** Acquire the draw buffer and release it at the end of the scope. */
 	struct FScopedAcquireDrawBuffer
 	{
-		FScopedAcquireDrawBuffer(ISlate3DRenderer& InSlateRenderer)
+		FScopedAcquireDrawBuffer(ISlate3DRenderer& InSlateRenderer, bool bInDeferRenderTargetUpdate = false)
 			: SlateRenderer(InSlateRenderer)
 			, DrawBuffer(InSlateRenderer.AcquireDrawBuffer())
+			, bDeferRenderTargetUpdate(bInDeferRenderTargetUpdate)
 		{
 		}
 		~FScopedAcquireDrawBuffer()
 		{
-			SlateRenderer.ReleaseDrawBuffer(DrawBuffer);
+			if (!bDeferRenderTargetUpdate)
+			{
+				SlateRenderer.ReleaseDrawBuffer(DrawBuffer);
+			}
 		}
 		FScopedAcquireDrawBuffer(const FScopedAcquireDrawBuffer&) = delete;
 		FScopedAcquireDrawBuffer& operator=(const FScopedAcquireDrawBuffer&) = delete;
@@ -37,6 +41,7 @@ public:
 	private:
 		ISlate3DRenderer& SlateRenderer;
 		FSlateDrawBuffer& DrawBuffer;
+		bool bDeferRenderTargetUpdate;
 	};
 
 	/** set if this renderer should render in gamma space by default. */
diff --git a/Engine/Source/Runtime/UMG/Private/Slate/WidgetRenderer.cpp b/Engine/Source/Runtime/UMG/Private/Slate/WidgetRenderer.cpp
--- a/Engine/Source/Runtime/UMG/Private/Slate/WidgetRenderer.cpp
+++ b/Engine/Source/Runtime/UMG/Private/Slate/WidgetRenderer.cpp
@@ -236,7 +236,7 @@ void FWidgetRenderer::DrawWindow(
 
 		{
 			// Get the free buffer & add our virtual window
-			ISlate3DRenderer::FScopedAcquireDrawBuffer ScopedDrawBuffer{ *Renderer };
+			ISlate3DRenderer::FScopedAcquireDrawBuffer ScopedDrawBuffer{ *Renderer, bDeferRenderTargetUpdate };
 			FSlateWindowElementList& WindowElementList = ScopedDrawBuffer.GetDrawBuffer().AddWindowElementList(Window);
 
 			// Paint the window
@@ -301,7 +301,7 @@ bool FWidgetRenderer::DrawInvalidationRoot(TSharedRef<SVirtualWindow>& VirtualWi
 
 		{
 			// Get the free buffer & add our virtual window
-			ISlate3DRenderer::FScopedAcquireDrawBuffer ScopedDrawBuffer{ *Renderer };
+			ISlate3DRenderer::FScopedAcquireDrawBuffer ScopedDrawBuffer{ *Renderer, bDeferRenderTargetUpdate };
 			FSlateWindowElementList& WindowElementList = ScopedDrawBuffer.GetDrawBuffer().AddWindowElementList(VirtualWindow);
 
 			ContextCopy.WindowElementList = &WindowElementList;
