Fixing issues in Slate after last checkin.

diff --git a/Engine/Source/Runtime/SlateRHIRenderer/Private/SlateRHIRenderingPolicy.cpp b/Engine/Source/Runtime/SlateRHIRenderer/Private/SlateRHIRenderingPolicy.cpp
--- a/Engine/Source/Runtime/SlateRHIRenderer/Private/SlateRHIRenderingPolicy.cpp
+++ b/Engine/Source/Runtime/SlateRHIRenderer/Private/SlateRHIRenderingPolicy.cpp
@@ -474,6 +474,10 @@ static bool UpdateScissorRect(
 
 
 						SetGraphicsPipelineState(RHICmdList, WriteMaskPSOInit, 0);
+
+						FRHIBatchedShaderParameters& BatchedParameters = RHICmdList.GetScratchShaderParameters();
+						VertexShader->SetViewProjection(BatchedParameters, FMatrix44f(ViewProjection));
+						RHICmdList.SetBatchedShaderParameters(VertexShader.GetVertexShader(), BatchedParameters);
 					}
 				}
 
@@ -1053,7 +1057,7 @@ void FSlateRHIRenderingPolicy::DrawElements(
 				}
 				{
 					GlobalVertexShader->SetViewProjection(BatchedParameters, FMatrix44f(ViewProjection));
-					RHICmdList.SetBatchedShaderParameters(PixelShader.GetPixelShader(), BatchedParameters);
+					RHICmdList.SetBatchedShaderParameters(GlobalVertexShader.GetVertexShader(), BatchedParameters);
 				}
 
 				{
