Fix invalid shader material resource crash slate

diff --git a/Engine/Source/Runtime/SlateCore/Private/Styling/SlateBrush.cpp b/Engine/Source/Runtime/SlateCore/Private/Styling/SlateBrush.cpp
--- a/Engine/Source/Runtime/SlateCore/Private/Styling/SlateBrush.cpp
+++ b/Engine/Source/Runtime/SlateCore/Private/Styling/SlateBrush.cpp
@@ -161,6 +161,17 @@ void FSlateBrush::UpdateRenderingResource(FVector2f LocalSize, float DrawScale)
 		{
 			ResourceHandle = FSlateApplicationBase::Get().GetRenderer()->GetResourceHandle(*this, LocalSize, DrawScale);
 		}
+		else if (ResourceHandle.IsValid())
+		{
+			// Test the resource itself
+			if (FSlateShaderResource* Resource = ResourceHandle.GetResourceProxy()->Resource)
+			{
+				if (!Resource->IsResourceValid())
+				{
+					ResourceHandle = FSlateApplicationBase::Get().GetRenderer()->GetResourceHandle(*this, LocalSize, DrawScale);
+				}
+			}
+		}
 	}
 
 }
diff --git a/Engine/Source/Runtime/SlateCore/Public/Textures/SlateShaderResource.h b/Engine/Source/Runtime/SlateCore/Public/Textures/SlateShaderResource.h
--- a/Engine/Source/Runtime/SlateCore/Public/Textures/SlateShaderResource.h
+++ b/Engine/Source/Runtime/SlateCore/Public/Textures/SlateShaderResource.h
@@ -64,6 +64,11 @@ public:
 	 */
 	virtual ESlateShaderResource::Type GetType() const = 0;
 
+	/**
+	 * Additional validation that can vary per resource type
+	 */
+	virtual bool IsResourceValid() const { return true; };
+
 #if SLATE_CHECK_UOBJECT_RENDER_RESOURCES
 	virtual void CheckForStaleResources() const { }
 	bool Debug_IsDestroyed() const { return DestroyState != 0x21; }
diff --git a/Engine/Source/Runtime/SlateRHIRenderer/Private/SlateMaterialResource.cpp b/Engine/Source/Runtime/SlateRHIRenderer/Private/SlateMaterialResource.cpp
--- a/Engine/Source/Runtime/SlateRHIRenderer/Private/SlateMaterialResource.cpp
+++ b/Engine/Source/Runtime/SlateRHIRenderer/Private/SlateMaterialResource.cpp
@@ -101,6 +101,17 @@ FSlateMaterialResource::~FSlateMaterialResource()
 	}
 }
 
+bool FSlateMaterialResource::IsResourceValid() const
+{
+	if (MaterialProxy && MaterialObject)
+	{
+		// Conservatively only return invalid in scenarios where proxy is mismatched
+		return MaterialProxy == MaterialObject->GetRenderProxy();
+	}
+
+	return true;
+}
+
 void FSlateMaterialResource::UpdateMaterial(const UMaterialInterface& InMaterialResource, const FVector2f InImageSize, FSlateShaderResource* InTextureMask)
 {
 #if SLATE_CHECK_UOBJECT_RENDER_RESOURCES
diff --git a/Engine/Source/Runtime/SlateRHIRenderer/Private/SlateMaterialResource.h b/Engine/Source/Runtime/SlateRHIRenderer/Private/SlateMaterialResource.h
--- a/Engine/Source/Runtime/SlateRHIRenderer/Private/SlateMaterialResource.h
+++ b/Engine/Source/Runtime/SlateRHIRenderer/Private/SlateMaterialResource.h
@@ -17,9 +17,12 @@ public:
 	FSlateMaterialResource(const UMaterialInterface& InMaterialResource, const FVector2f InImageSize, FSlateShaderResource* InTextureMask = nullptr );
 	~FSlateMaterialResource();
 
+	//~ Begin FSlateShaderResource Interface.
 	virtual uint32 GetWidth() const override { return Width; }
 	virtual uint32 GetHeight() const override { return Height; }
 	virtual ESlateShaderResource::Type GetType() const override { return ESlateShaderResource::Material; }
+	virtual bool IsResourceValid() const override;
+	//~ End FSlateShaderResource Interface.
 
 	void UpdateMaterial(const UMaterialInterface& InMaterialResource, const FVector2f InImageSize, FSlateShaderResource* InTextureMask );
 	void ResetMaterial();
