Fix unexpected GPU cost due to full screen stencil-only passes.

Some RHIs (such as D3D12) reset scissor and viewport state when beginning the render pass, so the previous optimization that configured the scissor around the UI element was ineffective.

diff --git a/Engine/Source/Runtime/SlateRHIRenderer/Private/SlateRHIRenderingPolicy.cpp b/Engine/Source/Runtime/SlateRHIRenderer/Private/SlateRHIRenderingPolicy.cpp
--- a/Engine/Source/Runtime/SlateRHIRenderer/Private/SlateRHIRenderingPolicy.cpp
+++ b/Engine/Source/Runtime/SlateRHIRenderer/Private/SlateRHIRenderingPolicy.cpp
@@ -343,11 +343,35 @@ static bool UpdateScissorRect(
 
 				// We only clear the stencil the first time, and if some how the user draws more than 255 masking quads
 				// in a single frame.
-				bool bClearStencil = false;
-				if (MaskingID == 0)
+				const bool bClearStencil = MaskingID == 0;
+
+				// Don't bother setting the render targets unless we actually need to clear them.
+				if (bClearStencil || bForceStateChange)
 				{
-					bClearStencil = true;
+					// #todo-renderpasses Similar to above this is gross. Would require a refactor to really fix.
+					RHICmdList.EndRenderPass();
+					bDidRestartRenderpass = true;
+
+					// Clear current stencil buffer, we use ELoad/EStore, because we need to keep the stencil around.
+					ERenderTargetLoadAction StencilLoadAction = bClearStencil ? ERenderTargetLoadAction::EClear : ERenderTargetLoadAction::ELoad;
+					ERenderTargetActions StencilAction = MakeRenderTargetActions(StencilLoadAction, ERenderTargetStoreAction::EStore);
+					if (IsMemorylessTexture(DepthStencilTarget))
+					{
+						// We can't preserve content for memoryless targets
+						StencilAction = bClearStencil ? ERenderTargetActions::Clear_DontStore : ERenderTargetActions::DontLoad_DontStore;
+					}
+
+					RPInfo.DepthStencilRenderTarget.Action = MakeDepthStencilTargetActions(ERenderTargetActions::DontLoad_DontStore, StencilAction);
+					RPInfo.DepthStencilRenderTarget.DepthStencilTarget = DepthStencilTarget;
+					RPInfo.DepthStencilRenderTarget.ExclusiveDepthStencil = FExclusiveDepthStencil::DepthNop_StencilWrite;
+					TransitionRenderPassTargets(RHICmdList, RPInfo);
+					RHICmdList.BeginRenderPass(RPInfo, TEXT("SlateUpdateScissorRect_ClearStencil"));
+				}
+
+				// Setup the scissor rect after starting the render pass, as the RHI does not preserve the scissor state between passes / render targets.
 
+				if (bClearStencil)
+				{
 					// We don't want there to be any scissor rect when we clear the stencil
 					RHICmdList.SetScissorRect(false, 0, 0, 0, 0);
 				}
@@ -375,28 +399,6 @@ static bool UpdateScissorRect(
 					RHICmdList.SetScissorRect(true, ScissorRect.Left, ScissorRect.Top, ScissorRect.Right, ScissorRect.Bottom);
 				}
 
-				// Don't bother setting the render targets unless we actually need to clear them.
-				if (bClearStencil || bForceStateChange)
-				{
-					// #todo-renderpasses Similar to above this is gross. Would require a refactor to really fix.
-					RHICmdList.EndRenderPass();
-					bDidRestartRenderpass = true;
-
-					// Clear current stencil buffer, we use ELoad/EStore, because we need to keep the stencil around.
-					ERenderTargetLoadAction StencilLoadAction = bClearStencil ? ERenderTargetLoadAction::EClear : ERenderTargetLoadAction::ELoad;
-					ERenderTargetActions StencilAction = MakeRenderTargetActions(StencilLoadAction, ERenderTargetStoreAction::EStore);
-					if (IsMemorylessTexture(DepthStencilTarget))
-					{
-						// We can't preserve content for memoryless targets
-						StencilAction = bClearStencil ? ERenderTargetActions::Clear_DontStore : ERenderTargetActions::DontLoad_DontStore;
-					}
-
-					RPInfo.DepthStencilRenderTarget.Action = MakeDepthStencilTargetActions(ERenderTargetActions::DontLoad_DontStore, StencilAction);
-					RPInfo.DepthStencilRenderTarget.DepthStencilTarget = DepthStencilTarget;
-					RPInfo.DepthStencilRenderTarget.ExclusiveDepthStencil = FExclusiveDepthStencil::DepthNop_StencilWrite;
-					TransitionRenderPassTargets(RHICmdList, RPInfo);
-					RHICmdList.BeginRenderPass(RPInfo, TEXT("SlateUpdateScissorRect_ClearStencil"));
-				}
 
 				FGlobalShaderMap* MaxFeatureLevelShaderMap = GetGlobalShaderMap(GMaxRHIShaderPlatform);
 
