Speculative fix slate texture resource crashes

diff --git a/Engine/Source/Runtime/SlateRHIRenderer/Private/SlateUTextureResource.cpp b/Engine/Source/Runtime/SlateRHIRenderer/Private/SlateUTextureResource.cpp
--- a/Engine/Source/Runtime/SlateRHIRenderer/Private/SlateUTextureResource.cpp
+++ b/Engine/Source/Runtime/SlateRHIRenderer/Private/SlateUTextureResource.cpp
@@ -97,18 +97,28 @@ void FSlateUTextureResource::UpdateTexture(UTexture* InTexture)
 	UpdateDebugName();
 #endif
 
-	if (!Proxy)
+	if (!Proxy && TextureObject)
 	{
 		Proxy = new FSlateShaderResourceProxy;
 	}
 
-	FTexture* TextureResource = InTexture->GetResource();
+	if (Proxy && TextureObject)
+	{
+		FTexture* TextureResource = TextureObject->GetResource();
 
-	Proxy->Resource = this;
-	// If the RHI data has changed, it's possible the underlying size of the texture has changed,
-	// if that's true we need to update the actual size recorded on the proxy as well, otherwise 
-	// the texture will continue to render using the wrong size.
-	Proxy->ActualSize = FIntPoint(TextureResource->GetSizeX(), TextureResource->GetSizeY());
+		Proxy->Resource = this;
+		// If the RHI data has changed, it's possible the underlying size of the texture has changed,
+		// if that's true we need to update the actual size recorded on the proxy as well, otherwise 
+		// the texture will continue to render using the wrong size.
+		if (TextureResource)
+		{
+			Proxy->ActualSize = FIntPoint(TextureResource->GetSizeX(), TextureResource->GetSizeY());
+		}
+		else
+		{
+			Proxy->ActualSize = FIntPoint(0, 0);
+		}
+	}
 }
 
 void FSlateUTextureResource::ResetTexture()
