Possible MSDF bug fixes for console and mobile

#rnx <Release Notes exclude>

diff --git a/Engine/Source/Runtime/SlateRHIRenderer/Private/SlateRHIFontTexture.cpp b/Engine/Source/Runtime/SlateRHIRenderer/Private/SlateRHIFontTexture.cpp
--- a/Engine/Source/Runtime/SlateRHIRenderer/Private/SlateRHIFontTexture.cpp
+++ b/Engine/Source/Runtime/SlateRHIRenderer/Private/SlateRHIFontTexture.cpp
@@ -103,8 +103,9 @@ EPixelFormat FSlateFontTextureRHIResource::GetRHIPixelFormat() const
 		case ESlateFontAtlasContentType::Alpha:
 			return PF_A8;
 		case ESlateFontAtlasContentType::Color:
-		case ESlateFontAtlasContentType::Msdf:
 			return PF_B8G8R8A8;
+		case ESlateFontAtlasContentType::Msdf:
+			return PF_R8G8B8A8;
 		default:
 			checkNoEntry();
 			// Default to Color
@@ -116,6 +117,11 @@ FSlateFontAtlasRHI::FSlateFontAtlasRHI(uint32 Width, uint32 Height, ESlateFontAt
 	: FSlateFontAtlas(Width, Height, InContentType, InPaddingStyle)
 	, FontTexture(new FSlateFontTextureRHIResource(Width, Height, InContentType))
 {
+	if (InContentType == ESlateFontAtlasContentType::Msdf)
+	{
+		// Actually this should be done for all content types but to be safe, I want to avoid affecting non-MSDF code for now.
+		bNeedsUpdate = true;
+	}
 }
 
 FSlateFontAtlasRHI::~FSlateFontAtlasRHI()
