Fix ScreenResolution not mapping to viewport in PIE for UI materials / widgets.

diff --git a/Engine/Source/Runtime/SlateRHIRenderer/Private/SlateRHIRenderer.cpp b/Engine/Source/Runtime/SlateRHIRenderer/Private/SlateRHIRenderer.cpp
--- a/Engine/Source/Runtime/SlateRHIRenderer/Private/SlateRHIRenderer.cpp
+++ b/Engine/Source/Runtime/SlateRHIRenderer/Private/SlateRHIRenderer.cpp
@@ -119,6 +119,7 @@ struct FSlateDrawWindowCommandParams
 	FSlateRHIRenderer* Renderer;
 	FSlateWindowElementList* WindowElementList;
 	SWindow* Window;
+	FIntRect ViewRect;
 #if WANTS_DRAW_MESH_EVENTS
 	FString WindowTitle;
 #endif
@@ -973,6 +974,7 @@ void RenderSlateBatch(FTexture2DRHIRef SlateRenderTarget, bool bClear, bool bIsH
 				RenderParams.bWireFrame = !!SlateWireFrame;
 				RenderParams.bIsHDR = bIsHDR;
 				RenderParams.HDRDisplayColorGamut = ViewportInfo.HDRDisplayColorGamut;
+				RenderParams.ViewRect = DrawCommandParams.ViewRect;
 				if (ViewportInfo.bSceneHDREnabled && !bIsHDR)
 				{
 					RenderParams.UITarget = ViewportInfo.UITargetRT;
@@ -1717,11 +1719,35 @@ void FSlateRHIRenderer::DrawWindows_Private(FSlateDrawBuffer& WindowDrawBuffer)
 
 				// Tell the rendering thread to draw the windows
 				{
+					auto GetViewRect = [Window]()
+					{
+#if WITH_EDITOR
+						if (GIsEditor)
+						{
+							if (TSharedPtr<ISlateViewport> Viewport = Window->GetViewport())
+							{
+								if (TSharedPtr<SWidget> ViewportWidget = Viewport->GetWidget().Pin())
+								{
+									// The actual backbuffer has a padding that extends beyond the draw area, account for this in our offsets
+									int32 OffsetX = FMath::RoundToInt32(ViewportWidget->GetTickSpaceGeometry().GetAbsolutePosition().X - Window->GetPositionInScreen().X);
+									int32 OffsetY = FMath::RoundToInt32(ViewportWidget->GetTickSpaceGeometry().GetAbsolutePosition().Y - Window->GetPositionInScreen().Y);
+
+									FIntPoint ViewportSize = Viewport->GetSize();
+									return FIntRect(FIntPoint(OffsetX, OffsetY), FIntPoint(ViewportSize.X + OffsetX, ViewportSize.Y + OffsetY));
+								}
+							}
+						}
+#endif // WITH_EDITOR
+
+						return FIntRect();
+					};
+
 					FSlateDrawWindowCommandParams Params;
 
 					Params.Renderer = this;
 					Params.WindowElementList = &ElementList;
 					Params.Window = Window;
+					Params.ViewRect = GetViewRect();
 #if WANTS_DRAW_MESH_EVENTS
 					Params.WindowTitle = Window->GetTitle().ToString();
 #endif
diff --git a/Engine/Source/Runtime/SlateRHIRenderer/Private/SlateRHIRenderingPolicy.cpp b/Engine/Source/Runtime/SlateRHIRenderer/Private/SlateRHIRenderingPolicy.cpp
--- a/Engine/Source/Runtime/SlateRHIRenderer/Private/SlateRHIRenderingPolicy.cpp
+++ b/Engine/Source/Runtime/SlateRHIRenderer/Private/SlateRHIRenderingPolicy.cpp
@@ -194,7 +194,7 @@ void FSlateRHIRenderingPolicy::BuildRenderingBuffers(FRHICommandListImmediate& R
 	SET_DWORD_STAT(STAT_SlateVertexCount, InBatchData.GetFinalVertexData().Num());
 }
 
-static FSceneView* CreateSceneView( FSceneViewFamilyContext* ViewFamilyContext, FSlateBackBuffer& BackBuffer, const FMatrix& ViewProjectionMatrix )
+static FSceneView* CreateSceneView( FSceneViewFamilyContext* ViewFamilyContext, FSlateBackBuffer& BackBuffer, const FMatrix& ViewProjectionMatrix, const FIntRect InViewRect)
 {
 	QUICK_SCOPE_CYCLE_COUNTER(STAT_Slate_CreateSceneView);
 	// In loading screens, the engine is NULL, so we skip out.
@@ -202,8 +202,12 @@ static FSceneView* CreateSceneView( FSceneViewFamilyContext* ViewFamilyContext,
 	{
 		return nullptr;
 	}
-
-	FIntRect ViewRect(FIntPoint(0, 0), BackBuffer.GetSizeXY());
+	
+	FIntRect ViewRect = InViewRect;
+	if (ViewRect.IsEmpty())
+	{
+		ViewRect = FIntRect(FIntPoint(0, 0), BackBuffer.GetSizeXY());
+	}
 
 	// make a temporary view
 	FSceneViewInitOptions ViewInitOptions;
@@ -749,7 +753,7 @@ void FSlateRHIRenderingPolicy::DrawElements(
 				.SetTime(Params.Time)
 				.SetRealtimeUpdate(true)
 			);
-			SceneViews[i] = CreateSceneView(SceneViewFamilyContexts[i], BackBuffer, FMatrix(Params.ViewProjectionMatrix));
+			SceneViews[i] = CreateSceneView(SceneViewFamilyContexts[i], BackBuffer, FMatrix(Params.ViewProjectionMatrix), Params.ViewRect);
 		}
 
 		SceneViewFamilyContexts[NumScenes - 1] = new FSceneViewFamilyContext
@@ -763,7 +767,7 @@ void FSlateRHIRenderingPolicy::DrawElements(
 			.SetTime(Params.Time)
 			.SetRealtimeUpdate(true)
 		);
-		SceneViews[NumScenes - 1] = CreateSceneView(SceneViewFamilyContexts[NumScenes - 1], BackBuffer, FMatrix(Params.ViewProjectionMatrix));
+		SceneViews[NumScenes - 1] = CreateSceneView(SceneViewFamilyContexts[NumScenes - 1], BackBuffer, FMatrix(Params.ViewProjectionMatrix), Params.ViewRect);
 	}
 
 	TShaderMapRef<FSlateElementVS> GlobalVertexShader(GetGlobalShaderMap(GMaxRHIShaderPlatform));
diff --git a/Engine/Source/Runtime/SlateRHIRenderer/Private/SlateRHIRenderingPolicy.h b/Engine/Source/Runtime/SlateRHIRenderer/Private/SlateRHIRenderingPolicy.h
--- a/Engine/Source/Runtime/SlateRHIRenderer/Private/SlateRHIRenderingPolicy.h
+++ b/Engine/Source/Runtime/SlateRHIRenderer/Private/SlateRHIRenderingPolicy.h
@@ -32,6 +32,7 @@ struct FSlateRenderingParams
 {
 	FMatrix44f ViewProjectionMatrix;
 	FVector2f ViewOffset;
+	FIntRect ViewRect;
 	FGameTime Time;
 	TRefCountPtr<IPooledRenderTarget> UITarget;
 	EDisplayColorGamut HDRDisplayColorGamut;
@@ -41,6 +42,7 @@ struct FSlateRenderingParams
 	FSlateRenderingParams(const FMatrix& InViewProjectionMatrix, FGameTime InTime)
 		: ViewProjectionMatrix(InViewProjectionMatrix)
 		, ViewOffset(0.f, 0.f)
+		, ViewRect(FIntRect())
 		, Time(InTime)
 		, HDRDisplayColorGamut(EDisplayColorGamut::sRGB_D65)
 		, bWireFrame(false)
