Expose scenebuffers to UMG
---
Feature guarded by 'Slate.CopyBackbufferToSlatePostRenderTargets' disabled by default in initial commit.
To use simply call GetSlatePost0 material function after enabling CVar.
No overhead unless a widget is rendered using the material function, overhead is minimal (Single-copy 0.003 ms).
GetSlatePost1-4 require enabling via project settings.

#virtualized

diff --git a/Engine/Source/Runtime/SlateCore/Private/Rendering/ElementBatcher.cpp b/Engine/Source/Runtime/SlateCore/Private/Rendering/ElementBatcher.cpp
--- a/Engine/Source/Runtime/SlateCore/Private/Rendering/ElementBatcher.cpp
+++ b/Engine/Source/Runtime/SlateCore/Private/Rendering/ElementBatcher.cpp
@@ -283,6 +283,7 @@ FSlateElementBatcher::FSlateElementBatcher( TSharedRef<FSlateRenderingPolicy> In
 	, bSRGBVertexColor( !InRenderingPolicy->IsVertexColorInLinearSpace() )
 	, bRequiresVsync(false)
 	, bCompositeHDRViewports(false)
+	, UsedSlatePostBuffers(ESlatePostRT::None)
 {
 }
 
@@ -3090,6 +3091,11 @@ FSlateRenderBatch& FSlateElementBatcher::CreateRenderBatch(
 
 	NewBatch.ClippingState = ClippingState;
 
+	if (InResource)
+	{
+		UsedSlatePostBuffers |= InResource->GetUsedSlatePostBuffers();
+	}
+
 	return NewBatch;
 }
 
@@ -3110,6 +3116,11 @@ FSlateRenderBatch& FSlateElementBatcher::CreateRenderBatch(
 
 	NewBatch.ClippingState = ResolveClippingState(DrawElement);
 
+	if (InResource)
+	{
+		UsedSlatePostBuffers |= InResource->GetUsedSlatePostBuffers();
+	}
+
 	return NewBatch;
 }
 
@@ -3561,6 +3572,7 @@ void FSlateElementBatcher::ResetBatches()
 {
 	bRequiresVsync = false;
 	bCompositeHDRViewports = false;
+	UsedSlatePostBuffers = ESlatePostRT::None;
 	NumPostProcessPasses = 0;
 }
 
diff --git a/Engine/Source/Runtime/SlateCore/Private/Rendering/SlateRendererSettings.cpp b/Engine/Source/Runtime/SlateCore/Private/Rendering/SlateRendererSettings.cpp
new file mode 100644
--- /dev/null
+++ b/Engine/Source/Runtime/SlateCore/Private/Rendering/SlateRendererSettings.cpp
@@ -0,0 +1,70 @@
+// Copyright Epic Games, Inc. All Rights Reserved.
+
+#include "Rendering/SlateRendererSettings.h"
+#include "HAL/IConsoleManager.h"
+
+#include UE_INLINE_GENERATED_CPP_BY_NAME(SlateRendererSettings)
+
+static TAutoConsoleVariable<int32> CVarDefaultEnablePostRenderTarget_0(
+	TEXT("Slate.DefaultEnablePostRenderTarget_0"),
+	1,
+	TEXT("Experimental. Set true to enable slate post render target 0"),
+	ECVF_ReadOnly);
+
+FSlatePostSettings::FSlatePostSettings()
+	: bEnabled(false)
+	, PathToSlatePostRT(FString())
+	, CachedSlatePostRT(nullptr)
+{
+}
+
+USlateRendererSettings::USlateRendererSettings()
+{
+	SlatePostSettings.Add(ESlatePostRT::ESlatePostRT_0, FSlatePostSettings());
+	SlatePostSettings.Add(ESlatePostRT::ESlatePostRT_1, FSlatePostSettings());
+	SlatePostSettings.Add(ESlatePostRT::ESlatePostRT_2, FSlatePostSettings());
+	SlatePostSettings.Add(ESlatePostRT::ESlatePostRT_3, FSlatePostSettings());
+	SlatePostSettings.Add(ESlatePostRT::ESlatePostRT_4, FSlatePostSettings());
+
+	// By default, enable the first post RT
+	SlatePostSettings[ESlatePostRT::ESlatePostRT_0].bEnabled = CVarDefaultEnablePostRenderTarget_0.GetValueOnAnyThread();
+
+	// Hardcoded paths to engine assets
+	SlatePostSettings[ESlatePostRT::ESlatePostRT_0].PathToSlatePostRT = "/Engine/EngineResources/SlatePost0_RT.SlatePost0_RT";
+	SlatePostSettings[ESlatePostRT::ESlatePostRT_1].PathToSlatePostRT = "/Engine/EngineResources/SlatePost1_RT.SlatePost1_RT";
+	SlatePostSettings[ESlatePostRT::ESlatePostRT_2].PathToSlatePostRT = "/Engine/EngineResources/SlatePost2_RT.SlatePost2_RT";
+	SlatePostSettings[ESlatePostRT::ESlatePostRT_3].PathToSlatePostRT = "/Engine/EngineResources/SlatePost3_RT.SlatePost3_RT";
+	SlatePostSettings[ESlatePostRT::ESlatePostRT_4].PathToSlatePostRT = "/Engine/EngineResources/SlatePost4_RT.SlatePost4_RT";
+}
+
+USlateRendererSettings::~USlateRendererSettings()
+{
+	for (ESlatePostRT SlatePostBufferBit : TEnumRange<ESlatePostRT>())
+	{
+		UObject* SlatePostBuffer = SlatePostSettings[SlatePostBufferBit].CachedSlatePostRT;
+		if (SlatePostBuffer)
+		{
+			SlatePostBuffer->RemoveFromRoot();
+		}
+	}
+}
+
+UObject* USlateRendererSettings::TryGetPostBufferRT(ESlatePostRT InPostBufferBit) const
+{
+	return SlatePostSettings[InPostBufferBit].CachedSlatePostRT;
+}
+
+UObject* USlateRendererSettings::LoadGetPostBufferRT(ESlatePostRT InPostBufferBit)
+{
+	UObject* Result = TryGetPostBufferRT(InPostBufferBit);
+
+	if (!Result)
+	{
+		Result = LoadObject<UObject>(nullptr, *SlatePostSettings[InPostBufferBit].PathToSlatePostRT, nullptr, LOAD_None, nullptr);
+		Result->AddToRoot();
+
+		SlatePostSettings[InPostBufferBit].CachedSlatePostRT = Result;
+	}
+
+	return Result;
+}
diff --git a/Engine/Source/Runtime/SlateCore/Public/Rendering/ElementBatcher.h b/Engine/Source/Runtime/SlateCore/Public/Rendering/ElementBatcher.h
--- a/Engine/Source/Runtime/SlateCore/Public/Rendering/ElementBatcher.h
+++ b/Engine/Source/Runtime/SlateCore/Public/Rendering/ElementBatcher.h
@@ -4,6 +4,7 @@
 
 #include "CoreMinimal.h"
 #include "Rendering/RenderingCommon.h"
+#include "Rendering/SlateRendererSettings.h"
 #include "Layout/Clipping.h"
 #include "Stats/Stats.h"
 #include "SlateGlobals.h"
@@ -255,6 +256,8 @@ public:
 
 	bool CompositeHDRViewports() const { return bCompositeHDRViewports; }
 
+	ESlatePostRT GetUsedSlatePostBuffers() const { return UsedSlatePostBuffers; }
+
 	void SetCompositeHDRViewports(bool bInCompositeHDRViewports) { bCompositeHDRViewports = bInCompositeHDRViewports; }
 
 	/** 
@@ -480,4 +483,7 @@ private:
 
 	// true if viewports get composited as a separate pass, instead of being rendered directly to the render target. Useful for HDR displays
 	bool bCompositeHDRViewports;
+
+	// true if we added a resource that is using a slate post buffer 
+	ESlatePostRT UsedSlatePostBuffers;
 };
diff --git a/Engine/Source/Runtime/SlateCore/Public/Rendering/SlateRendererSettings.h b/Engine/Source/Runtime/SlateCore/Public/Rendering/SlateRendererSettings.h
new file mode 100644
--- /dev/null
+++ b/Engine/Source/Runtime/SlateCore/Public/Rendering/SlateRendererSettings.h
@@ -0,0 +1,111 @@
+// Copyright Epic Games, Inc. All Rights Reserved.
+
+#pragma once
+
+#include "CoreMinimal.h"
+#include "Containers/ContainerAllocationPolicies.h"
+#include "Engine/DeveloperSettings.h"
+#include "Misc/EnumRange.h"
+#include "UObject/Object.h"
+#include "UObject/SoftObjectPath.h"
+
+#include "SlateRendererSettings.generated.h"
+
+/**
+ * Bitfield used to mark if a slate post RT is used or not
+ */
+UENUM(BlueprintType)
+enum class ESlatePostRT : uint8
+{
+	None = 0 << 0,
+	ESlatePostRT_0 = 1 << 0,
+	ESlatePostRT_1 = 1 << 1,
+	ESlatePostRT_2 = 1 << 2,
+	ESlatePostRT_3 = 1 << 3,
+	ESlatePostRT_4 = 1 << 4,
+	Num = 5
+};
+
+ENUM_CLASS_FLAGS(ESlatePostRT);
+
+ENUM_RANGE_BY_VALUES(ESlatePostRT, ESlatePostRT::ESlatePostRT_0, ESlatePostRT::ESlatePostRT_1, ESlatePostRT::ESlatePostRT_2, ESlatePostRT::ESlatePostRT_3, ESlatePostRT::ESlatePostRT_4);
+
+/**
+ * Settings for a particular Slate Post RT.
+ * Notably if enabled & blur by default. To be updated with additional effects & to be expandable in game code / settings.
+ */
+USTRUCT(BlueprintType, meta=(HiddenByDefault, DisableSplitPin))
+struct SLATECORE_API FSlatePostSettings
+{
+	GENERATED_BODY()
+
+public:
+
+	FSlatePostSettings();
+
+	friend class USlateRendererSettings;
+
+public:
+
+	UPROPERTY(EditAnywhere, BlueprintReadWrite, Category=Overrides, meta=(PinHiddenByDefault))
+	uint8 bEnabled:1;
+
+public:
+
+	const FString& GetPathToSlatePostRT() const { return PathToSlatePostRT; }
+
+private:
+
+	/** Path to Slate Post RT Asset */
+	FString PathToSlatePostRT;
+
+	/** Cached load of Slate Post RT Asset */
+	UObject* CachedSlatePostRT;
+};
+
+/**
+ * Settings used to control slate rendering
+ */
+UCLASS(config = Game, defaultconfig)
+class SLATECORE_API USlateRendererSettings : public UDeveloperSettings
+{
+	GENERATED_BODY()
+
+public:
+		
+	static const USlateRendererSettings* Get() 
+	{ 
+		return GetDefault<USlateRendererSettings>();
+	}
+
+	static USlateRendererSettings* GetMutable()
+	{
+		return GetMutableDefault<USlateRendererSettings>();
+	}
+
+public:
+
+	USlateRendererSettings();
+	~USlateRendererSettings();
+
+	/** Try to get post RT asset, returns nullptr if not already loaded */
+	UObject* TryGetPostBufferRT(ESlatePostRT InPostBufferBit) const;
+
+	/** Get post RT asset, loading if not already loaded */
+	UObject* LoadGetPostBufferRT(ESlatePostRT InPostBufferBit);
+
+public:
+
+	/** 
+	 * Map of all slate post RT's and their settings 
+	 * Note that each post RT used in a frame will result in 1 full framebuffer copy for slate to sample from.
+	 * If a post RT is not used, no copy will occur & that post RT will be resized to 1x1 after 2 frames of non-use.
+	 * 
+	 * By default only SlatePostRT_0 is enabled. The rest must manually be enabled in settings below.
+	 */
+
+	// @TODO: I did a map, kinda hoping small map optimizations would be fine, anyone suggest different data structure? Discuss in review.
+	// Map is nice since needs no editor customization. After initial run there should be no more than 5 lookups each frame.
+	UPROPERTY(EditAnywhere, EditFixedSize, Category = "PostProcessing")
+	TMap<ESlatePostRT, FSlatePostSettings> SlatePostSettings;
+};
diff --git a/Engine/Source/Runtime/SlateCore/Public/Textures/SlateShaderResource.h b/Engine/Source/Runtime/SlateCore/Public/Textures/SlateShaderResource.h
--- a/Engine/Source/Runtime/SlateCore/Public/Textures/SlateShaderResource.h
+++ b/Engine/Source/Runtime/SlateCore/Public/Textures/SlateShaderResource.h
@@ -5,6 +5,7 @@
 #include "CoreMinimal.h"
 #include "Stats/Stats.h"
 #include "SlateGlobals.h"
+#include "Rendering/SlateRendererSettings.h"
 #include "Rendering/SlateResourceHandle.h"
 
 class FSlateShaderResourceProxy;
@@ -69,6 +70,11 @@ public:
 	 */
 	virtual bool IsResourceValid() const { return true; };
 
+	/**
+	 * Does this resource use slate post buffers? If no resources use a post buffer we won't populate it.
+	 */
+	virtual ESlatePostRT GetUsedSlatePostBuffers() const { return ESlatePostRT::None; }
+
 #if SLATE_CHECK_UOBJECT_RENDER_RESOURCES
 	virtual void CheckForStaleResources() const { }
 	bool Debug_IsDestroyed() const { return DestroyState != 0x21; }
diff --git a/Engine/Source/Runtime/SlateCore/SlateCore.Build.cs b/Engine/Source/Runtime/SlateCore/SlateCore.Build.cs
--- a/Engine/Source/Runtime/SlateCore/SlateCore.Build.cs
+++ b/Engine/Source/Runtime/SlateCore/SlateCore.Build.cs
@@ -12,6 +12,7 @@ public class SlateCore : ModuleRules
 			new string[] {
 				"Core",
 				"CoreUObject",
+				"DeveloperSettings",
 				"InputCore",
 				"Json",
 				"TraceLog",
diff --git a/Engine/Source/Runtime/SlateRHIRenderer/Private/SlateMaterialResource.cpp b/Engine/Source/Runtime/SlateRHIRenderer/Private/SlateMaterialResource.cpp
--- a/Engine/Source/Runtime/SlateRHIRenderer/Private/SlateMaterialResource.cpp
+++ b/Engine/Source/Runtime/SlateRHIRenderer/Private/SlateMaterialResource.cpp
@@ -1,6 +1,8 @@
 // Copyright Epic Games, Inc. All Rights Reserved.
 
 #include "SlateMaterialResource.h"
+#include "Rendering/SlateRendererSettings.h"
+#include "Engine/Texture.h"
 #include "Materials/MaterialInstanceDynamic.h"
 #include "Materials/MaterialRenderProxy.h"
 #include "MaterialShared.h"
@@ -70,6 +72,7 @@ FSlateMaterialResource::FSlateMaterialResource(const UMaterialInterface& InMater
 	, TextureMaskResource( InTextureMask )
 	, Width(FMath::RoundToInt(InImageSize.X))
 	, Height(FMath::RoundToInt(InImageSize.Y))
+	, CachedSlatePostBuffers(ESlatePostRT::None)
 {
 #if SLATE_CHECK_UOBJECT_RENDER_RESOURCES
 	SlateMaterialResource::CheckInvalidUMaterial(InMaterialResource, NAME_None);
@@ -101,6 +104,11 @@ FSlateMaterialResource::~FSlateMaterialResource()
 	}
 }
 
+ESlatePostRT FSlateMaterialResource::GetUsedSlatePostBuffers() const
+{
+	return CachedSlatePostBuffers;
+}
+
 bool FSlateMaterialResource::IsResourceValid() const
 {
 	if (MaterialProxy && MaterialObject)
@@ -123,6 +131,28 @@ void FSlateMaterialResource::UpdateMaterial(const UMaterialInterface& InMaterial
 	MaterialObjectWeakPtr = MaterialObject;
 	UpdateMaterialName();
 
+	if (MaterialObject)
+	{
+		// Quality / Feature level irrelevant since flag to search all levels for both is true
+		TArray<UTexture*> OutUsedTextures;
+		MaterialObject->GetUsedTextures(OutUsedTextures, EMaterialQualityLevel::Num, true, ERHIFeatureLevel::Num, true);
+
+		CachedSlatePostBuffers = ESlatePostRT::None;
+		for (const UTexture* OutUsedTexture : OutUsedTextures)
+		{
+			for (const TPair<ESlatePostRT, FSlatePostSettings>& SlatePostSetting : USlateRendererSettings::Get()->SlatePostSettings)
+			{
+				const ESlatePostRT SlatePostBitflag = SlatePostSetting.Key;
+				const FSlatePostSettings& SlatePostSettingValue = SlatePostSetting.Value;
+
+				if (SlatePostSettingValue.bEnabled && OutUsedTexture && OutUsedTexture->GetPathName() == SlatePostSettingValue.GetPathToSlatePostRT())
+				{
+					CachedSlatePostBuffers |= SlatePostBitflag;
+				}
+			}
+		}
+	}
+
 	SlateMaterialResource::CheckInvalidMaterialProxy(MaterialProxy, DebugName);
 
 #else
@@ -160,6 +190,8 @@ void FSlateMaterialResource::ResetMaterial()
 	UpdateMaterialName();
 #endif
 
+	CachedSlatePostBuffers = ESlatePostRT::None;
+
 	TextureMaskResource = nullptr;
 	if (SlateProxy)
 	{
diff --git a/Engine/Source/Runtime/SlateRHIRenderer/Private/SlateMaterialResource.h b/Engine/Source/Runtime/SlateRHIRenderer/Private/SlateMaterialResource.h
--- a/Engine/Source/Runtime/SlateRHIRenderer/Private/SlateMaterialResource.h
+++ b/Engine/Source/Runtime/SlateRHIRenderer/Private/SlateMaterialResource.h
@@ -21,6 +21,7 @@ public:
 	virtual uint32 GetWidth() const override { return Width; }
 	virtual uint32 GetHeight() const override { return Height; }
 	virtual ESlateShaderResource::Type GetType() const override { return ESlateShaderResource::Material; }
+	virtual ESlatePostRT GetUsedSlatePostBuffers() const override;
 	virtual bool IsResourceValid() const override;
 	//~ End FSlateShaderResource Interface.
 
@@ -62,6 +63,9 @@ private:
 	uint32 Width;
 	uint32 Height;
 
+	/** Cached slate SlatePostRT assets / buffers in use */
+	ESlatePostRT CachedSlatePostBuffers;
+
 private:
 #if SLATE_CHECK_UOBJECT_RENDER_RESOURCES
 	void UpdateMaterialName();
diff --git a/Engine/Source/Runtime/SlateRHIRenderer/Private/SlateRHIRenderer.cpp b/Engine/Source/Runtime/SlateRHIRenderer/Private/SlateRHIRenderer.cpp
--- a/Engine/Source/Runtime/SlateRHIRenderer/Private/SlateRHIRenderer.cpp
+++ b/Engine/Source/Runtime/SlateRHIRenderer/Private/SlateRHIRenderer.cpp
@@ -9,6 +9,8 @@
 #include "Widgets/SWindow.h"
 #include "Framework/Application/SlateApplication.h"
 #include "EngineGlobals.h"
+#include "Engine/AssetManager.h"
+#include "Engine/TextureRenderTarget2D.h"
 #include "Materials/MaterialRenderProxy.h"
 #include "MaterialShared.h"
 #include "RendererInterface.h"
@@ -21,9 +23,16 @@
 #include "ScreenRendering.h"
 #include "SlateShaders.h"
 #include "Rendering/ElementBatcher.h"
+#include "Rendering/SlateRenderer.h"
+#include "Rendering/SlateRendererSettings.h"
+#include "RenderResource.h"
+#include "RenderingThread.h"
+#include "RHIResources.h"
+#include "RHIUtilities.h"
 #include "StereoRendering.h"
 #include "SlateNativeTextureResource.h"
 #include "SceneUtils.h"
+#include "TextureResource.h"
 #include "VolumeRendering.h"
 #include "ShaderCompiler.h"
 #include "PipelineStateCache.h"
@@ -88,6 +97,12 @@ static TAutoConsoleVariable<int32> CVarMemorylessDepthStencil(
 	TEXT("Whether to use memoryless DepthStencil target for Slate. Reduces memory usage and implies that DepthStencil state can't be preserved between Slate renderpasses"),
 	ECVF_RenderThreadSafe);
 
+static TAutoConsoleVariable<int32> CVarCopyBackbufferToSlatePostRenderTargets(
+	TEXT("Slate.CopyBackbufferToSlatePostRenderTargets"),
+	0,
+	TEXT("Experimental. Set true to copy final backbuffer into slate RTs for slate post processing / material usage"),
+	ECVF_RenderThreadSafe);
+
 #if WITH_SLATE_VISUALIZERS
 
 TAutoConsoleVariable<int32> CVarShowSlateOverdraw(
@@ -224,6 +239,7 @@ FSlateRHIRenderer::FSlateRHIRenderer(TSharedRef<FSlateFontServices> InSlateFontS
 	OutHDRScreenshotData = NULL;
 	ScreenshotViewportInfo = nullptr;
 	bIsStandaloneStereoOnlyDevice = IHeadMountedDisplayModule::IsAvailable() && IHeadMountedDisplayModule::Get().IsStandaloneStereoOnlyDevice();
+	bShrinkPostBufferRequested = ESlatePostRT::None;
 }
 
 FSlateRHIRenderer::~FSlateRHIRenderer()
@@ -1704,6 +1720,9 @@ void FSlateRHIRenderer::DrawWindows_Private(FSlateDrawBuffer& WindowDrawBuffer)
 
 				bLockToVsync |= bForceVsyncFromCVar;
 
+				// Cache if the element batcher used a post buffer, will get reset on batch reset
+				ESlatePostRT UsedSlatePostBufferBits = ElementBatcher->GetUsedSlatePostBuffers();
+
 				// All elements for this window have been batched and rendering data updated
 				ElementBatcher->ResetBatches();
 
@@ -1717,6 +1736,106 @@ void FSlateRHIRenderer::DrawWindows_Private(FSlateDrawBuffer& WindowDrawBuffer)
 					ConditionalResizeViewport(ViewInfo, ViewInfo->DesiredWidth, ViewInfo->DesiredHeight, IsViewportFullscreen(*Window), Window);
 				}
 
+				// Update slate post buffers before slate draw if needed
+				if (CVarCopyBackbufferToSlatePostRenderTargets.GetValueOnAnyThread() && IsInGameThread() && GIsClient && !IsRunningCommandlet() && !GUsingNullRHI && UAssetManager::IsInitialized())
+				{
+					uint8 SlatePostBufferBitIndex = 0;
+					for (ESlatePostRT SlatePostBufferBit : TEnumRange<ESlatePostRT>())
+					{
+						if (!USlateRendererSettings::Get()->SlatePostSettings[SlatePostBufferBit].bEnabled)
+						{
+							SlatePostBufferBitIndex++;
+							continue;
+						}
+
+						UTextureRenderTarget2D* SlatePostBuffer = Cast<UTextureRenderTarget2D>(USlateRendererSettings::Get()->TryGetPostBufferRT(SlatePostBufferBit));
+						if (!SlatePostBuffer)
+						{
+							SlatePostBuffer = Cast<UTextureRenderTarget2D>(USlateRendererSettings::GetMutable()->LoadGetPostBufferRT(SlatePostBufferBit));
+						}
+
+						bool bIsViewportPresentForPIE = GIsEditor ? Window->GetViewport().IsValid() : true;
+						bool bPostBufferBitUsed = (UsedSlatePostBufferBits & SlatePostBufferBit) != ESlatePostRT::None;
+
+						if (bIsViewportPresentForPIE && bPostBufferBitUsed)
+						{
+							// Viewport only has a RT texture in editor / PIE, this texture will be the size of the viewport
+							// typically this texture 'BufferedRT', is smaller than the entire backbuffer used for the editor
+							// We initialize to nullptr so this ptr can be copied & resolved here instead of in the render command
+							FSlateRenderTargetRHI* ViewportTexture = nullptr;
+							if (GIsEditor)
+							{
+								ViewportTexture = static_cast<FSlateRenderTargetRHI*>(Window->GetViewport()->GetViewportRenderTargetTexture());
+							}
+							
+							FIntPoint SizeSlatePostRT = GIsEditor
+								? FIntPoint(ViewportTexture->GetWidth(), ViewportTexture->GetHeight())
+								: FIntPoint(ViewInfo->DesiredWidth, ViewInfo->DesiredHeight);
+
+							bool bHDREnabled = IsHDREnabled();
+							bool bIsPixelFormatCorrect = bHDREnabled 
+								? SlatePostBuffer->GetFormat() == EPixelFormat::PF_FloatRGBA 
+								: SlatePostBuffer->GetFormat() == EPixelFormat::PF_A2B10G10R10;
+							if (SlatePostBuffer->SizeX != SizeSlatePostRT.X || SlatePostBuffer->SizeY != SizeSlatePostRT.Y || !bIsPixelFormatCorrect)
+							{
+								SlatePostBuffer->InitCustomFormat(SizeSlatePostRT.X, SizeSlatePostRT.Y, bHDREnabled ? EPixelFormat::PF_FloatRGBA : EPixelFormat::PF_A2B10G10R10, false);
+							}
+
+							const FVector2D ElementWindowSize = ElementList.GetWindowSize();
+							ENQUEUE_RENDER_COMMAND(FUpdateSlatePostBuffers)([ViewInfo, ElementWindowSize, ViewportTexture, SlatePostBuffer](FRHICommandListImmediate& RHICmdList)
+							{
+								bool bRenderedStereo = false;
+								if (CVarDrawToVRRenderTarget->GetInt() == 0 && GEngine && IsValidRef(ViewInfo->GetRenderTargetTexture()) && GEngine->StereoRenderingDevice.IsValid())
+								{
+									GEngine->StereoRenderingDevice->RenderTexture_RenderThread(RHICmdList, RHIGetViewportBackBuffer(ViewInfo->ViewportRHI), ViewInfo->GetRenderTargetTexture(), ElementWindowSize);
+									bRenderedStereo = true;
+								}
+
+								FTexture2DRHIRef ViewportRT = bRenderedStereo ? nullptr : ViewInfo->GetRenderTargetTexture();
+								FTexture2DRHIRef BackBuffer = (ViewportRT) ? ViewportRT : RHIGetViewportBackBuffer(ViewInfo->ViewportRHI);
+
+								if (BackBuffer)
+								{
+									FRHICopyTextureInfo CopyInfo;
+
+									// Copy just the viewport RT if in PIE, else do entire backbuffer
+									if (GIsEditor)
+									{
+										CopyInfo.Size = FIntVector(ViewportTexture->GetWidth(), ViewportTexture->GetHeight(), 1);
+										TransitionAndCopyTexture(RHICmdList, ViewportTexture->GetRHIRef(), SlatePostBuffer->TextureReference.TextureReferenceRHI, CopyInfo);
+									}
+									else
+									{
+										FIntPoint BackbufferExtent = BackBuffer->GetDesc().Extent;
+										CopyInfo.Size = FIntVector(BackbufferExtent.X, BackbufferExtent.Y, 1);
+										TransitionAndCopyTexture(RHICmdList, BackBuffer, SlatePostBuffer->TextureReference.TextureReferenceRHI, CopyInfo);
+									}
+								}
+							});
+
+							SlatePostRTFences[SlatePostBufferBitIndex].BeginFence();
+
+							bShrinkPostBufferRequested &= ~SlatePostBufferBit;
+						}
+						else if (SlatePostBuffer->GetResource() && SlatePostRTFences[SlatePostBufferBitIndex].IsFenceComplete() && (SlatePostBuffer->SizeX != 1 || SlatePostBuffer->SizeY != 1))
+						{
+							if ((bShrinkPostBufferRequested & SlatePostBufferBit) == ESlatePostRT::None)
+							{
+								// Delay shrink attempts for a frame, since if we resize while there is an active copy in flight, we will crash
+								bShrinkPostBufferRequested |= SlatePostBufferBit;
+							}
+							else
+							{
+								// Resize unused SlatePostRTs to 1x1.
+								SlatePostBuffer->InitCustomFormat(1, 1, IsHDREnabled() ? EPixelFormat::PF_FloatRGBA : EPixelFormat::PF_A2B10G10R10, false);
+								bShrinkPostBufferRequested &= ~SlatePostBufferBit;
+							}
+						}
+
+						SlatePostBufferBitIndex++;
+					}
+				}
+
 				// Tell the rendering thread to draw the windows
 				{
 					auto GetViewRect = [Window]()
diff --git a/Engine/Source/Runtime/SlateRHIRenderer/Private/SlateRHIRenderer.h b/Engine/Source/Runtime/SlateRHIRenderer/Private/SlateRHIRenderer.h
--- a/Engine/Source/Runtime/SlateRHIRenderer/Private/SlateRHIRenderer.h
+++ b/Engine/Source/Runtime/SlateRHIRenderer/Private/SlateRHIRenderer.h
@@ -6,10 +6,12 @@
 #include "Textures/SlateShaderResource.h"
 #include "Rendering/DrawElements.h"
 #include "RHI.h"
+#include "RenderCommandFence.h"
 #include "RenderResource.h"
 #include "SlateRHIResourceManager.h"
 #include "UnrealClient.h"
 #include "Rendering/SlateRenderer.h"
+#include "Rendering/SlateRendererSettings.h"
 #include "Rendering/SlateDrawBuffer.h"
 #include "Slate/SlateTextures.h"
 #include "RendererInterface.h"
@@ -305,6 +307,8 @@ private:
 	bool bIsStandaloneStereoOnlyDevice;
 	bool bTakingAScreenShot;
 	bool bUpdateHDRDisplayInformation;
+	ESlatePostRT bShrinkPostBufferRequested;
+	FRenderCommandFence SlatePostRTFences[(uint8)ESlatePostRT::Num];
 	FIntRect ScreenshotRect;
 	FViewportInfo* ScreenshotViewportInfo;
 	TArray<FColor>* OutScreenshotData;
diff --git a/Engine/Source/Runtime/SlateRHIRenderer/Private/SlateUTextureResource.cpp b/Engine/Source/Runtime/SlateRHIRenderer/Private/SlateUTextureResource.cpp
--- a/Engine/Source/Runtime/SlateRHIRenderer/Private/SlateUTextureResource.cpp
+++ b/Engine/Source/Runtime/SlateRHIRenderer/Private/SlateUTextureResource.cpp
@@ -4,6 +4,7 @@
 
 FSlateBaseUTextureResource::FSlateBaseUTextureResource(UTexture* InTexture)
 	: TextureObject(InTexture)
+	, CachedSlatePostBuffers(ESlatePostRT::None)
 {
 #if SLATE_CHECK_UOBJECT_RENDER_RESOURCES
 	ObjectWeakPtr = InTexture;
@@ -30,6 +31,11 @@ ESlateShaderResource::Type FSlateBaseUTextureResource::GetType() const
 	return ESlateShaderResource::TextureObject;
 }
 
+ESlatePostRT FSlateBaseUTextureResource::GetUsedSlatePostBuffers() const
+{
+	return CachedSlatePostBuffers;
+}
+
 #if SLATE_CHECK_UOBJECT_RENDER_RESOURCES
 void FSlateBaseUTextureResource::UpdateDebugName()
 {
@@ -104,6 +110,18 @@ void FSlateUTextureResource::UpdateTexture(UTexture* InTexture)
 
 	if (Proxy && TextureObject)
 	{
+		CachedSlatePostBuffers = ESlatePostRT::None;
+		for (const TPair<ESlatePostRT, FSlatePostSettings>& SlatePostSetting : USlateRendererSettings::Get()->SlatePostSettings)
+		{
+			const ESlatePostRT SlatePostBitflag = SlatePostSetting.Key;
+			const FSlatePostSettings& SlatePostSettingValue = SlatePostSetting.Value;
+
+			if (SlatePostSettingValue.bEnabled && InTexture && InTexture->GetPathName() == SlatePostSettingValue.GetPathToSlatePostRT())
+			{
+				CachedSlatePostBuffers |= SlatePostBitflag;
+			}
+		}
+
 		FTexture* TextureResource = TextureObject->GetResource();
 
 		Proxy->Resource = this;
@@ -130,6 +148,8 @@ void FSlateUTextureResource::ResetTexture()
 	UpdateDebugName();
 #endif
 
+	CachedSlatePostBuffers = ESlatePostRT::None;
+
 	if (Proxy)
 	{
 		delete Proxy;
diff --git a/Engine/Source/Runtime/SlateRHIRenderer/Private/SlateUTextureResource.h b/Engine/Source/Runtime/SlateRHIRenderer/Private/SlateUTextureResource.h
--- a/Engine/Source/Runtime/SlateRHIRenderer/Private/SlateUTextureResource.h
+++ b/Engine/Source/Runtime/SlateRHIRenderer/Private/SlateUTextureResource.h
@@ -21,6 +21,7 @@ public:
 	virtual uint32 GetWidth() const;
 	virtual uint32 GetHeight() const;
 	virtual ESlateShaderResource::Type GetType() const;
+	virtual ESlatePostRT GetUsedSlatePostBuffers() const override;
 
 	/** Gets the RHI resource used for rendering and updates the last render time for texture streaming */
 	FTextureRHIRef AccessRHIResource()
@@ -56,6 +57,9 @@ protected:
 	TWeakObjectPtr<UTexture> ObjectWeakPtr;
 	FName DebugName;
 #endif
+
+	/** Cached slate SlatePostRT assets / buffers in use */
+	ESlatePostRT CachedSlatePostBuffers;
 };
 
 
diff --git a/Engine/Source/Runtime/SlateRHIRenderer/SlateRHIRenderer.Build.cs b/Engine/Source/Runtime/SlateRHIRenderer/SlateRHIRenderer.Build.cs
--- a/Engine/Source/Runtime/SlateRHIRenderer/SlateRHIRenderer.Build.cs
+++ b/Engine/Source/Runtime/SlateRHIRenderer/SlateRHIRenderer.Build.cs
@@ -10,6 +10,7 @@ public class SlateRHIRenderer : ModuleRules
 			new string[] {
 				"Core",
 				"CoreUObject",
+				"DeveloperSettings",
 				"Slate",
 				"SlateCore",
                 "Engine",
