Add framework for processing slate scenebuffer copies
---
To create a new slate post buffer processor, inherit from 'USlateRHIPostBufferProcessor' & implement 'PostProcess'
To use in game, create a BP that derives from your buffer processor, ex 'USlatePostBufferBlur'.
Set default values on that BP, then connect the processor to a buffer in 'SlateRendererSettings'.

#tests perf pie standalone shipping hdr-non-composite

diff --git a/Engine/Source/Runtime/SlateCore/Private/Rendering/SlateRendererSettings.cpp b/Engine/Source/Runtime/SlateCore/Private/Rendering/SlateRendererSettings.cpp
--- a/Engine/Source/Runtime/SlateCore/Private/Rendering/SlateRendererSettings.cpp
+++ b/Engine/Source/Runtime/SlateCore/Private/Rendering/SlateRendererSettings.cpp
@@ -13,11 +13,24 @@ static TAutoConsoleVariable<int32> CVarDefaultEnablePostRenderTarget_0(
 
 FSlatePostSettings::FSlatePostSettings()
 	: bEnabled(false)
+	, PostProcessorClass(nullptr)
 	, PathToSlatePostRT(FString())
 	, CachedSlatePostRT(nullptr)
 {
 }
 
+USlateCorePostBufferProcessor* FSlatePostSettings::GetProcessor() const
+{
+	USlateCorePostBufferProcessor* Result = nullptr;
+
+	if (PostProcessorClass)
+	{
+		Result = Cast<USlateCorePostBufferProcessor>(PostProcessorClass->GetDefaultObject());
+	}
+
+	return Result;
+}
+
 USlateRendererSettings::USlateRendererSettings()
 {
 	SlatePostSettings.Add(ESlatePostRT::ESlatePostRT_0, FSlatePostSettings());
@@ -39,9 +52,11 @@ USlateRendererSettings::USlateRendererSettings()
 
 USlateRendererSettings::~USlateRendererSettings()
 {
-	for (ESlatePostRT SlatePostBufferBit : TEnumRange<ESlatePostRT>())
+	for (TPair<ESlatePostRT, FSlatePostSettings>& SlatePostSetting : SlatePostSettings)
 	{
-		UObject* SlatePostBuffer = SlatePostSettings[SlatePostBufferBit].CachedSlatePostRT;
+		FSlatePostSettings& PostSetting = SlatePostSetting.Value;
+
+		UObject* SlatePostBuffer = PostSetting.CachedSlatePostRT;
 		if (SlatePostBuffer)
 		{
 			SlatePostBuffer->RemoveFromRoot();
@@ -49,6 +64,22 @@ USlateRendererSettings::~USlateRendererSettings()
 	}
 }
 
+FSlatePostSettings& USlateRendererSettings::GetMutableSlatePostSetting(ESlatePostRT InPostBufferBit)
+{
+	return SlatePostSettings[InPostBufferBit];
+}
+
+const FSlatePostSettings& USlateRendererSettings::GetSlatePostSetting(ESlatePostRT InPostBufferBit) const
+{
+	return SlatePostSettings[InPostBufferBit];
+}
+
+USlateCorePostBufferProcessor* USlateRendererSettings::GetSlatePostProcessor(ESlatePostRT InPostBufferBit) const
+{
+	const FSlatePostSettings& Settings = SlatePostSettings[InPostBufferBit];
+	return Settings.GetProcessor();
+}
+
 UObject* USlateRendererSettings::TryGetPostBufferRT(ESlatePostRT InPostBufferBit) const
 {
 	return SlatePostSettings[InPostBufferBit].CachedSlatePostRT;
@@ -61,10 +92,18 @@ UObject* USlateRendererSettings::LoadGetPostBufferRT(ESlatePostRT InPostBufferBi
 	if (!Result)
 	{
 		Result = LoadObject<UObject>(nullptr, *SlatePostSettings[InPostBufferBit].PathToSlatePostRT, nullptr, LOAD_None, nullptr);
-		Result->AddToRoot();
 
-		SlatePostSettings[InPostBufferBit].CachedSlatePostRT = Result;
+		if (Result)
+		{
+			Result->AddToRoot();
+			SlatePostSettings[InPostBufferBit].CachedSlatePostRT = Result;
+		}
 	}
 
 	return Result;
 }
+
+const TMap<ESlatePostRT, FSlatePostSettings>& USlateRendererSettings::GetSlatePostSettings() const
+{
+	return SlatePostSettings;
+}
diff --git a/Engine/Source/Runtime/SlateCore/Public/Rendering/SlateRendererSettings.h b/Engine/Source/Runtime/SlateCore/Public/Rendering/SlateRendererSettings.h
--- a/Engine/Source/Runtime/SlateCore/Public/Rendering/SlateRendererSettings.h
+++ b/Engine/Source/Runtime/SlateCore/Public/Rendering/SlateRendererSettings.h
@@ -8,6 +8,8 @@
 #include "Misc/EnumRange.h"
 #include "UObject/Object.h"
 #include "UObject/SoftObjectPath.h"
+#include "Templates/SubclassOf.h"
+
 
 #include "SlateRendererSettings.generated.h"
 
@@ -30,6 +32,21 @@ ENUM_CLASS_FLAGS(ESlatePostRT);
 
 ENUM_RANGE_BY_VALUES(ESlatePostRT, ESlatePostRT::ESlatePostRT_0, ESlatePostRT::ESlatePostRT_1, ESlatePostRT::ESlatePostRT_2, ESlatePostRT::ESlatePostRT_3, ESlatePostRT::ESlatePostRT_4);
 
+/**
+ * Do not inherit from. Instead inherit from USlateRHIPostBufferProcessor. For an example see: USlatePostBufferBlur.
+ * 
+ * Base class for types that can process the backbuffer scene into the slate post buffer.
+ * This class is exposed to SlateCore, but due to module limitations, you should inherit
+ * from 'USlateRHIPostBufferProcessor' & implement 'PostProcess_RenderThread' in your derived class.
+ * 
+ * SlateRHI will only know how to utilize classes that derive from USlateRHIPostBufferProcessor.
+ */
+UCLASS(Abstract, Blueprintable, CollapseCategories)
+class SLATECORE_API USlateCorePostBufferProcessor : public UObject
+{
+	GENERATED_BODY()
+};
+
 /**
  * Settings for a particular Slate Post RT.
  * Notably if enabled & blur by default. To be updated with additional effects & to be expandable in game code / settings.
@@ -47,17 +64,29 @@ public:
 
 public:
 
-	UPROPERTY(EditAnywhere, BlueprintReadWrite, Category=Overrides, meta=(PinHiddenByDefault))
+	/** Should this post buffer be enabled for updating */
+	UPROPERTY(EditAnywhere, BlueprintReadWrite, Category=BufferSettings, meta=(PinHiddenByDefault))
 	uint8 bEnabled:1;
 
-public:
-
-	const FString& GetPathToSlatePostRT() const { return PathToSlatePostRT; }
+	/** Copy of actually loaded post processor class */
+	UPROPERTY(EditAnywhere, BlueprintReadWrite, Category = BufferSettings, Meta = (AllowAbstract = false))
+	TSubclassOf<USlateCorePostBufferProcessor> PostProcessorClass;
 
 private:
 
 	/** Path to Slate Post RT Asset */
+	UPROPERTY()
 	FString PathToSlatePostRT;
+	
+public:
+
+	/** Get post processing object using CDO from soft path if loaded */
+	USlateCorePostBufferProcessor* GetProcessor() const;
+
+	/** Get asset name for post RT texture */
+	const FString& GetPathToSlatePostRT() const { return PathToSlatePostRT; }
+
+private:
 
 	/** Cached load of Slate Post RT Asset */
 	UObject* CachedSlatePostRT;
@@ -88,13 +117,32 @@ public:
 	USlateRendererSettings();
 	~USlateRendererSettings();
 
+public:
+
+	/** Get settings struct for a particular post buffer index */
+	UFUNCTION(BlueprintCallable, Category = "SlateFX")
+	FSlatePostSettings& GetMutableSlatePostSetting(ESlatePostRT InPostBufferBit);
+
+	/** Get settings struct for a particular post buffer index */
+	UFUNCTION(BlueprintCallable, Category = "SlateFX")
+	const FSlatePostSettings& GetSlatePostSetting(ESlatePostRT InPostBufferBit) const;
+
+	/** Get post processor for a particular post buffer index, if it exists */
+	UFUNCTION(BlueprintCallable, Category = "SlateFX")
+	USlateCorePostBufferProcessor* GetSlatePostProcessor(ESlatePostRT InPostBufferBit) const;
+
+public:
+
 	/** Try to get post RT asset, returns nullptr if not already loaded */
 	UObject* TryGetPostBufferRT(ESlatePostRT InPostBufferBit) const;
 
 	/** Get post RT asset, loading if not already loaded */
 	UObject* LoadGetPostBufferRT(ESlatePostRT InPostBufferBit);
 
-public:
+	/** Get slate post settings map, non mutable */
+	const TMap<ESlatePostRT, FSlatePostSettings>& GetSlatePostSettings() const;
+
+private:
 
 	/** 
 	 * Map of all slate post RT's and their settings 
@@ -104,8 +152,7 @@ public:
 	 * By default only SlatePostRT_0 is enabled. The rest must manually be enabled in settings below.
 	 */
 
-	// @TODO: I did a map, kinda hoping small map optimizations would be fine, anyone suggest different data structure? Discuss in review.
 	// Map is nice since needs no editor customization. After initial run there should be no more than 5 lookups each frame.
-	UPROPERTY(EditAnywhere, EditFixedSize, Category = "PostProcessing")
+	UPROPERTY(config, EditAnywhere, EditFixedSize, Category = "PostProcessing")
 	TMap<ESlatePostRT, FSlatePostSettings> SlatePostSettings;
 };
diff --git a/Engine/Source/Runtime/SlateRHIRenderer/Private/FX/SlatePostBufferBlur.cpp b/Engine/Source/Runtime/SlateRHIRenderer/Private/FX/SlatePostBufferBlur.cpp
new file mode 100644
--- /dev/null
+++ b/Engine/Source/Runtime/SlateRHIRenderer/Private/FX/SlatePostBufferBlur.cpp
@@ -0,0 +1,34 @@
+// Copyright Epic Games, Inc. All Rights Reserved.
+
+#include "FX/SlatePostBufferBlur.h"
+#include "UObject/StrongObjectPtr.h"
+#include "UObject/WeakObjectPtrTemplates.h"
+#include "RHIResources.h"
+
+#include UE_INLINE_GENERATED_CPP_BY_NAME(SlatePostBufferBlur)
+
+void USlatePostBufferBlur::PostProcess(FRenderResource* InViewInfo, FRenderResource* InViewportTexture, FVector2D InElementWindowSize, FSlateRHIRenderingPolicyInterface InRenderingPolicy, UTextureRenderTarget2D* InSlatePostBuffer)
+{
+	if (InRenderingPolicy.IsValid())
+	{
+		// Explicit param copy to avoid renderthread from reading value during gamethread write
+		float GaussianBlurStrengthCopy = GaussianBlurStrength;
+
+		ENQUEUE_RENDER_COMMAND(FUpdateSlatePostBuffersWithFX_Blur)([Self = this, InViewInfo, InViewportTexture, InElementWindowSize, InRenderingPolicy, InSlatePostBuffer, GaussianBlurStrengthCopy](FRHICommandListImmediate& RHICmdList)
+		{
+			if (Self)
+			{
+				FTexture2DRHIRef BackBuffer = Self->GetBackbuffer_RenderThread(InViewInfo, InViewportTexture, InElementWindowSize, RHICmdList);
+
+				if (BackBuffer)
+				{
+					FTexture2DRHIRef Src = Self->GetSrcTexture_RenderThread(BackBuffer, InViewportTexture);
+					FTextureReferenceRHIRef& Dst = Self->GetDstTexture_RenderThread(InSlatePostBuffer);
+					FIntPoint DstExtent = Self->GetDstExtent_RenderThread(BackBuffer, InViewportTexture);
+
+					InRenderingPolicy.BlurRectExternal(RHICmdList, Src, Dst, DstExtent, GaussianBlurStrengthCopy);
+				}
+			}
+		});
+	}
+}
\ No newline at end of file
diff --git a/Engine/Source/Runtime/SlateRHIRenderer/Private/FX/SlateRHIPostBufferProcessor.cpp b/Engine/Source/Runtime/SlateRHIRenderer/Private/FX/SlateRHIPostBufferProcessor.cpp
new file mode 100644
--- /dev/null
+++ b/Engine/Source/Runtime/SlateRHIRenderer/Private/FX/SlateRHIPostBufferProcessor.cpp
@@ -0,0 +1,65 @@
+// Copyright Epic Games, Inc. All Rights Reserved.
+
+#include "FX/SlateRHIPostBufferProcessor.h"
+
+#include "RHIResources.h"
+#include "SlateRHIRenderer.h"
+#include "StereoRendering.h"
+#include "UnrealEngine.h"
+
+#include UE_INLINE_GENERATED_CPP_BY_NAME(SlateRHIPostBufferProcessor)
+
+void USlateRHIPostBufferProcessor::PostProcess(FRenderResource* InViewInfo, FRenderResource* InViewportTexture, FVector2D InElementWindowSize, FSlateRHIRenderingPolicyInterface InRenderingPolicy, UTextureRenderTarget2D* InSlatePostBuffer)
+{
+	ENQUEUE_RENDER_COMMAND(FUpdateSlatePostBuffersWithFX)([Self = this, InViewInfo, InViewportTexture, InElementWindowSize, InRenderingPolicy, InSlatePostBuffer](FRHICommandListImmediate& RHICmdList)
+	{
+		if (Self)
+		{
+			FTexture2DRHIRef BackBuffer = Self->GetBackbuffer_RenderThread(InViewInfo, InViewportTexture, InElementWindowSize, RHICmdList);
+
+			if (BackBuffer)
+			{
+				FTexture2DRHIRef Src = Self->GetSrcTexture_RenderThread(BackBuffer, InViewportTexture);
+				FTextureReferenceRHIRef& Dst = Self->GetDstTexture_RenderThread(InSlatePostBuffer);
+				FIntPoint DstExtent = Self->GetDstExtent_RenderThread(BackBuffer, InViewportTexture);
+
+				Self->PostProcess_RenderThread(InRenderingPolicy, RHICmdList, Src, Dst, DstExtent);
+			}
+		}
+	});
+}
+
+FTexture2DRHIRef USlateRHIPostBufferProcessor::GetBackbuffer_RenderThread(FRenderResource* InViewInfo, FRenderResource* InViewportTexture, FVector2D InElementWindowSize, FRHICommandListImmediate& InRHICmdList)
+{
+	FViewportInfo* ViewInfo = static_cast<FViewportInfo*>(InViewInfo);
+	FSlateRenderTargetRHI* ViewportTexture = static_cast<FSlateRenderTargetRHI*>(InViewportTexture);
+
+	bool bRenderedStereo = false;
+	if (FSlateRHIRenderer::GetDrawToVRRenderTarget() == 0 && GEngine && IsValidRef(ViewInfo->GetRenderTargetTexture()) && GEngine->StereoRenderingDevice.IsValid())
+	{
+		GEngine->StereoRenderingDevice->RenderTexture_RenderThread(InRHICmdList, RHIGetViewportBackBuffer(ViewInfo->ViewportRHI), ViewInfo->GetRenderTargetTexture(), InElementWindowSize);
+		bRenderedStereo = true;
+	}
+
+	FTexture2DRHIRef ViewportRT = bRenderedStereo ? nullptr : ViewInfo->GetRenderTargetTexture();
+	FTexture2DRHIRef BackBuffer = (ViewportRT) ? ViewportRT : RHIGetViewportBackBuffer(ViewInfo->ViewportRHI);
+
+	return BackBuffer;
+}
+
+FTexture2DRHIRef USlateRHIPostBufferProcessor::GetSrcTexture_RenderThread(FTexture2DRHIRef InBackBuffer, FRenderResource* InViewportTexture)
+{
+	FSlateRenderTargetRHI* ViewportTexture = static_cast<FSlateRenderTargetRHI*>(InViewportTexture);
+	return GIsEditor ? ViewportTexture->GetRHIRef() : InBackBuffer;
+}
+
+FTextureReferenceRHIRef& USlateRHIPostBufferProcessor::GetDstTexture_RenderThread(UTextureRenderTarget2D* InSlatePostBuffer)
+{
+	return InSlatePostBuffer->TextureReference.TextureReferenceRHI;
+}
+
+FIntPoint USlateRHIPostBufferProcessor::GetDstExtent_RenderThread(FTexture2DRHIRef InBackBuffer, FRenderResource* InViewportTexture)
+{
+	FSlateRenderTargetRHI* ViewportTexture = static_cast<FSlateRenderTargetRHI*>(InViewportTexture);
+	return GIsEditor ? FIntPoint(ViewportTexture->GetWidth(), ViewportTexture->GetHeight()) : InBackBuffer->GetDesc().Extent;
+}
\ No newline at end of file
diff --git a/Engine/Source/Runtime/SlateRHIRenderer/Private/SlateMaterialResource.cpp b/Engine/Source/Runtime/SlateRHIRenderer/Private/SlateMaterialResource.cpp
--- a/Engine/Source/Runtime/SlateRHIRenderer/Private/SlateMaterialResource.cpp
+++ b/Engine/Source/Runtime/SlateRHIRenderer/Private/SlateMaterialResource.cpp
@@ -131,6 +131,14 @@ void FSlateMaterialResource::UpdateMaterial(const UMaterialInterface& InMaterial
 	MaterialObjectWeakPtr = MaterialObject;
 	UpdateMaterialName();
 
+	SlateMaterialResource::CheckInvalidMaterialProxy(MaterialProxy, DebugName);
+
+#else
+
+	MaterialObject = &InMaterialResource;
+	MaterialProxy = InMaterialResource.GetRenderProxy();
+#endif
+
 	if (MaterialObject)
 	{
 		// Quality / Feature level irrelevant since flag to search all levels for both is true
@@ -140,7 +148,7 @@ void FSlateMaterialResource::UpdateMaterial(const UMaterialInterface& InMaterial
 		CachedSlatePostBuffers = ESlatePostRT::None;
 		for (const UTexture* OutUsedTexture : OutUsedTextures)
 		{
-			for (const TPair<ESlatePostRT, FSlatePostSettings>& SlatePostSetting : USlateRendererSettings::Get()->SlatePostSettings)
+			for (const TPair<ESlatePostRT, FSlatePostSettings>& SlatePostSetting : USlateRendererSettings::Get()->GetSlatePostSettings())
 			{
 				const ESlatePostRT SlatePostBitflag = SlatePostSetting.Key;
 				const FSlatePostSettings& SlatePostSettingValue = SlatePostSetting.Value;
@@ -153,14 +161,6 @@ void FSlateMaterialResource::UpdateMaterial(const UMaterialInterface& InMaterial
 		}
 	}
 
-	SlateMaterialResource::CheckInvalidMaterialProxy(MaterialProxy, DebugName);
-
-#else
-
-	MaterialObject = &InMaterialResource;
-	MaterialProxy = InMaterialResource.GetRenderProxy();
-#endif
-
 	if (MaterialProxy && (MaterialProxy->IsDeleted() || MaterialProxy->IsMarkedForGarbageCollection()))
 	{
 		MaterialProxy = nullptr;
diff --git a/Engine/Source/Runtime/SlateRHIRenderer/Private/SlatePostProcessor.cpp b/Engine/Source/Runtime/SlateRHIRenderer/Private/SlatePostProcessor.cpp
--- a/Engine/Source/Runtime/SlateRHIRenderer/Private/SlatePostProcessor.cpp
+++ b/Engine/Source/Runtime/SlateRHIRenderer/Private/SlatePostProcessor.cpp
@@ -144,11 +144,13 @@ static void BlitUIToHDRScene(FRHICommandListImmediate& RHICmdList, IRendererModu
 {
 	SCOPED_DRAW_EVENT(RHICmdList, SlatePostProcessBlitUIToHDR);
 
-	FRHITexture* UITexture = RectParams.UITarget->GetRHI();
+	FRHITexture* UITexture = RectParams.PostProcessDest == EPostProcessDestination::DestTexture 
+		? RectParams.DestTexture->GetTexture2D() 
+		: RectParams.UITarget->GetRHI();
 
 	RHICmdList.Transition(FRHITransitionInfo(UITexture, ERHIAccess::Unknown, ERHIAccess::SRVMask));
 	TRefCountPtr<IPooledRenderTarget> UITargetRTMask;
-	if (RHISupportsRenderTargetWriteMask(GMaxRHIShaderPlatform))
+	if (RHISupportsRenderTargetWriteMask(GMaxRHIShaderPlatform) && RectParams.PostProcessDest != EPostProcessDestination::DestTexture)
 	{
 		const auto FeatureLevel = GMaxRHIFeatureLevel;
 		auto ShaderMap = GetGlobalShaderMap(FeatureLevel);
@@ -631,7 +633,9 @@ void FSlatePostProcessor::UpsampleRect(FRHICommandListImmediate& RHICmdList, IRe
 	GraphicsPSOInit.DepthStencilState = TStaticDepthStencilState<false, CF_Always>::GetRHI();
 
 	// Original source texture is now the destination texture
-	FTexture2DRHIRef DestTexture = Params.SourceTexture;
+	FTexture2DRHIRef DestTexture = Params.PostProcessDest == EPostProcessDestination::DestTexture && Params.DestTexture 
+		? Params.DestTexture 
+		: Params.SourceTexture;
 	const int32 DestTextureWidth = Params.SourceTextureSize.X;
 	const int32 DestTextureHeight = Params.SourceTextureSize.Y;
 
@@ -665,6 +669,9 @@ void FSlatePostProcessor::UpsampleRect(FRHICommandListImmediate& RHICmdList, IRe
 	bool bIsSCRGB = false;
 
 	FRHITexture* UITargetTexture = Params.UITarget.IsValid() ? Params.UITarget->GetRHI() : nullptr;
+	UITargetTexture = Params.PostProcessDest == EPostProcessDestination::DestTexture && Params.DestTexture 
+		? Params.DestTexture->GetTexture2D() 
+		: UITargetTexture;
 
 	if (UITargetTexture != nullptr && DestTexture != UITargetTexture)
 	{
diff --git a/Engine/Source/Runtime/SlateRHIRenderer/Private/SlatePostProcessor.h b/Engine/Source/Runtime/SlateRHIRenderer/Private/SlatePostProcessor.h
--- a/Engine/Source/Runtime/SlateRHIRenderer/Private/SlatePostProcessor.h
+++ b/Engine/Source/Runtime/SlateRHIRenderer/Private/SlatePostProcessor.h
@@ -10,9 +10,18 @@ class FSlatePostProcessResource;
 class IRendererModule;
 struct IPooledRenderTarget;
 
+enum class EPostProcessDestination : uint8
+{
+	// Output postprocess to the default UI render target
+	UITarget,
+	// Output postprocess to the provided destination texture
+	DestTexture
+};
+
 struct FPostProcessRectParams
 {
 	FTextureRHIRef SourceTexture;
+	FTextureRHIRef DestTexture; // Only used when 'PostProcessDest' is 'DestTexture'
 	FSlateRect SourceRect;
 	FSlateRect DestRect;
 	FVector4f CornerRadius;
@@ -21,6 +30,7 @@ struct FPostProcessRectParams
 	TRefCountPtr<IPooledRenderTarget> UITarget; // not using FTextureRHIRef because we want to be able to use FRenderTargetWriteMask::Decode
 	uint32 StencilRef{};
 	EDisplayColorGamut HDRDisplayColorGamut;
+	EPostProcessDestination PostProcessDest = EPostProcessDestination::UITarget;
 };
 
 struct FBlurRectParams
diff --git a/Engine/Source/Runtime/SlateRHIRenderer/Private/SlateRHIRenderer.cpp b/Engine/Source/Runtime/SlateRHIRenderer/Private/SlateRHIRenderer.cpp
--- a/Engine/Source/Runtime/SlateRHIRenderer/Private/SlateRHIRenderer.cpp
+++ b/Engine/Source/Runtime/SlateRHIRenderer/Private/SlateRHIRenderer.cpp
@@ -11,6 +11,7 @@
 #include "EngineGlobals.h"
 #include "Engine/AssetManager.h"
 #include "Engine/TextureRenderTarget2D.h"
+#include "FX/SlateRHIPostBufferProcessor.h"
 #include "Materials/MaterialRenderProxy.h"
 #include "MaterialShared.h"
 #include "RendererInterface.h"
@@ -39,6 +40,7 @@
 #include "EngineModule.h"
 #include "Interfaces/ISlate3DRenderer.h"
 #include "SlateRHIRenderingPolicy.h"
+#include "Interfaces/SlateRHIRenderingPolicyInterface.h"
 #include "Slate/SlateTextureAtlasInterface.h"
 #include "Types/ReflectionMetadata.h"
 #include "CommonRenderResources.h"
@@ -265,6 +267,11 @@ FMatrix FSlateRHIRenderer::CreateProjectionMatrix(uint32 Width, uint32 Height)
 	);
 }
 
+int32 FSlateRHIRenderer::GetDrawToVRRenderTarget()
+{
+	return CVarDrawToVRRenderTarget->GetInt();
+}
+
 bool FSlateRHIRenderer::Initialize()
 {
 	LoadUsedTextures();
@@ -1742,7 +1749,7 @@ void FSlateRHIRenderer::DrawWindows_Private(FSlateDrawBuffer& WindowDrawBuffer)
 					uint8 SlatePostBufferBitIndex = 0;
 					for (ESlatePostRT SlatePostBufferBit : TEnumRange<ESlatePostRT>())
 					{
-						if (!USlateRendererSettings::Get()->SlatePostSettings[SlatePostBufferBit].bEnabled)
+						if (!USlateRendererSettings::Get()->GetSlatePostSetting(SlatePostBufferBit).bEnabled)
 						{
 							SlatePostBufferBitIndex++;
 							continue;
@@ -1782,42 +1789,51 @@ void FSlateRHIRenderer::DrawWindows_Private(FSlateDrawBuffer& WindowDrawBuffer)
 							}
 
 							const FVector2D ElementWindowSize = ElementList.GetWindowSize();
-							ENQUEUE_RENDER_COMMAND(FUpdateSlatePostBuffers)([ViewInfo, ElementWindowSize, ViewportTexture, SlatePostBuffer](FRHICommandListImmediate& RHICmdList)
-							{
-								bool bRenderedStereo = false;
-								if (CVarDrawToVRRenderTarget->GetInt() == 0 && GEngine && IsValidRef(ViewInfo->GetRenderTargetTexture()) && GEngine->StereoRenderingDevice.IsValid())
-								{
-									GEngine->StereoRenderingDevice->RenderTexture_RenderThread(RHICmdList, RHIGetViewportBackBuffer(ViewInfo->ViewportRHI), ViewInfo->GetRenderTargetTexture(), ElementWindowSize);
-									bRenderedStereo = true;
-								}
 
-								FTexture2DRHIRef ViewportRT = bRenderedStereo ? nullptr : ViewInfo->GetRenderTargetTexture();
-								FTexture2DRHIRef BackBuffer = (ViewportRT) ? ViewportRT : RHIGetViewportBackBuffer(ViewInfo->ViewportRHI);
-
-								if (BackBuffer)
+							if (USlateRHIPostBufferProcessor* PostProcessor = Cast<USlateRHIPostBufferProcessor>(USlateRendererSettings::Get()->GetSlatePostProcessor(SlatePostBufferBit)))
+							{
+								// Allow the post processor to enque render commands, we delgate this task so it can be done in a thread safe manner.
+								PostProcessor->PostProcess(ViewInfo, ViewportTexture, ElementWindowSize, FSlateRHIRenderingPolicyInterface(RenderingPolicy.Get()), SlatePostBuffer);
+							}
+							else
+							{
+								ENQUEUE_RENDER_COMMAND(FUpdateSlatePostBuffers)([ViewInfo, ElementWindowSize, ViewportTexture, SlatePostBuffer](FRHICommandListImmediate& RHICmdList)
 								{
-									FRHICopyTextureInfo CopyInfo;
-
-									// Copy just the viewport RT if in PIE, else do entire backbuffer
-									if (GIsEditor)
+									bool bRenderedStereo = false;
+									if (CVarDrawToVRRenderTarget->GetInt() == 0 && GEngine && IsValidRef(ViewInfo->GetRenderTargetTexture()) && GEngine->StereoRenderingDevice.IsValid())
 									{
-										CopyInfo.Size = FIntVector(ViewportTexture->GetWidth(), ViewportTexture->GetHeight(), 1);
-										TransitionAndCopyTexture(RHICmdList, ViewportTexture->GetRHIRef(), SlatePostBuffer->TextureReference.TextureReferenceRHI, CopyInfo);
+										GEngine->StereoRenderingDevice->RenderTexture_RenderThread(RHICmdList, RHIGetViewportBackBuffer(ViewInfo->ViewportRHI), ViewInfo->GetRenderTargetTexture(), ElementWindowSize);
+										bRenderedStereo = true;
 									}
-									else
+
+									FTexture2DRHIRef ViewportRT = bRenderedStereo ? nullptr : ViewInfo->GetRenderTargetTexture();
+									FTexture2DRHIRef BackBuffer = (ViewportRT) ? ViewportRT : RHIGetViewportBackBuffer(ViewInfo->ViewportRHI);
+
+									if (BackBuffer)
 									{
-										FIntPoint BackbufferExtent = BackBuffer->GetDesc().Extent;
-										CopyInfo.Size = FIntVector(BackbufferExtent.X, BackbufferExtent.Y, 1);
-										TransitionAndCopyTexture(RHICmdList, BackBuffer, SlatePostBuffer->TextureReference.TextureReferenceRHI, CopyInfo);
+										FRHICopyTextureInfo CopyInfo;
+
+										// Copy just the viewport RT if in PIE, else do entire backbuffer
+										if (GIsEditor)
+										{
+											CopyInfo.Size = FIntVector(ViewportTexture->GetWidth(), ViewportTexture->GetHeight(), 1);
+											TransitionAndCopyTexture(RHICmdList, ViewportTexture->GetRHIRef(), SlatePostBuffer->TextureReference.TextureReferenceRHI, CopyInfo);
+										}
+										else
+										{
+											FIntPoint BackbufferExtent = BackBuffer->GetDesc().Extent;
+											CopyInfo.Size = FIntVector(BackbufferExtent.X, BackbufferExtent.Y, 1);
+											TransitionAndCopyTexture(RHICmdList, BackBuffer, SlatePostBuffer->TextureReference.TextureReferenceRHI, CopyInfo);
+										}
 									}
-								}
-							});
+								});
+							}
 
 							SlatePostRTFences[SlatePostBufferBitIndex].BeginFence();
 
 							bShrinkPostBufferRequested &= ~SlatePostBufferBit;
 						}
-						else if (SlatePostBuffer->GetResource() && SlatePostRTFences[SlatePostBufferBitIndex].IsFenceComplete() && (SlatePostBuffer->SizeX != 1 || SlatePostBuffer->SizeY != 1))
+						else if (SlatePostBuffer && SlatePostBuffer->GetResource() && SlatePostRTFences[SlatePostBufferBitIndex].IsFenceComplete() && (SlatePostBuffer->SizeX != 1 || SlatePostBuffer->SizeY != 1))
 						{
 							if ((bShrinkPostBufferRequested & SlatePostBufferBit) == ESlatePostRT::None)
 							{
diff --git a/Engine/Source/Runtime/SlateRHIRenderer/Private/SlateRHIRenderer.h b/Engine/Source/Runtime/SlateRHIRenderer/Private/SlateRHIRenderer.h
--- a/Engine/Source/Runtime/SlateRHIRenderer/Private/SlateRHIRenderer.h
+++ b/Engine/Source/Runtime/SlateRHIRenderer/Private/SlateRHIRenderer.h
@@ -177,6 +177,7 @@ public:
 	 * @return The created projection matrix
 	 */
 	static FMatrix CreateProjectionMatrix( uint32 Width, uint32 Height );
+	static int32 GetDrawToVRRenderTarget();
 
 	/** FSlateRenderer interface */
 	virtual bool Initialize() override;
diff --git a/Engine/Source/Runtime/SlateRHIRenderer/Private/SlateRHIRenderingPolicy.cpp b/Engine/Source/Runtime/SlateRHIRenderer/Private/SlateRHIRenderingPolicy.cpp
--- a/Engine/Source/Runtime/SlateRHIRenderer/Private/SlateRHIRenderingPolicy.cpp
+++ b/Engine/Source/Runtime/SlateRHIRenderer/Private/SlateRHIRenderingPolicy.cpp
@@ -1653,4 +1653,64 @@ void FSlateRHIRenderingPolicy::FlushGeneratedResources()
 {
 	PostProcessor->ReleaseRenderTargets();
 }
+
+void FSlateRHIRenderingPolicy::BlurRectExternal(FRHICommandListImmediate& RHICmdList, FTexture2DRHIRef BlurSrc, FTextureReferenceRHIRef& BlurDst, FIntPoint DstExtent, float BlurStrength) const
+{
+	SLATE_DRAW_EVENT(RHICmdList, PostProcess);
+
+	FIntPoint BlurDstExtent = DstExtent;
+
+	// If the radius isn't set, auto-compute it based on the strength
+	int32 OutKernelSize = FMath::RoundToInt(BlurStrength * 3.f);
+
+	// Downsample if needed
+	int32 OutDownsampleAmount = 0;
+	if (OutKernelSize > 9)
+	{
+		OutDownsampleAmount = OutKernelSize >= 64 ? 4 : 2;
+		OutKernelSize /= OutDownsampleAmount;
+	}
+
+	// Kernel sizes must be odd
+	if (OutKernelSize % 2 == 0)
+	{
+		++OutKernelSize;
+	}
+
+	float ComputedStrength = FMath::Max(.5f, BlurStrength);
+
+	int32 RenderTargetWidth = BlurDstExtent.X;
+	int32 RenderTargetHeight = BlurDstExtent.Y;
+	
+	if (OutDownsampleAmount > 0)
+	{
+		RenderTargetWidth = FMath::DivideAndRoundUp(RenderTargetWidth, OutDownsampleAmount);
+		RenderTargetHeight = FMath::DivideAndRoundUp(RenderTargetHeight, OutDownsampleAmount);
+		ComputedStrength /= OutDownsampleAmount;
+	}
+
+	OutKernelSize = FMath::Clamp(OutKernelSize, 3, 255 /*MaxKernelSize*/);
+
+	FVector4f PostProcessData = FVector4f((float)OutKernelSize, ComputedStrength, (float)RenderTargetWidth, (float)RenderTargetHeight);
+
+	FVector2f TopLeft = FVector2f::ZeroVector;
+	FVector2f BotRight = FVector2f(BlurDstExtent.X, BlurDstExtent.Y);
+
+	FPostProcessRectParams RectParams;
+	RectParams.SourceTexture = BlurSrc;
+	RectParams.SourceRect = FSlateRect(0.f, 0.f, (float)BlurSrc->GetSizeX(), (float)BlurSrc->GetSizeY());
+	RectParams.DestRect = FSlateRect(TopLeft.X, TopLeft.Y, BotRight.X, BotRight.Y);
+	RectParams.SourceTextureSize = BlurSrc->GetSizeXY();
+	RectParams.CornerRadius = FVector4f(0, 0, 0, 0);
+	RectParams.DestTexture = BlurDst;
+	RectParams.PostProcessDest = EPostProcessDestination::DestTexture;
+
+	FBlurRectParams BlurParams;
+	BlurParams.KernelSize = PostProcessData.X;
+	BlurParams.Strength = PostProcessData.Y;
+	BlurParams.DownsampleAmount = OutDownsampleAmount;
+
+	IRendererModule& RendererModule = FModuleManager::GetModuleChecked<IRendererModule>(RendererModuleName);
+	PostProcessor->BlurRect(RHICmdList, RendererModule, BlurParams, RectParams);
+}
 	
diff --git a/Engine/Source/Runtime/SlateRHIRenderer/Private/SlateRHIRenderingPolicy.h b/Engine/Source/Runtime/SlateRHIRenderer/Private/SlateRHIRenderingPolicy.h
--- a/Engine/Source/Runtime/SlateRHIRenderer/Private/SlateRHIRenderingPolicy.h
+++ b/Engine/Source/Runtime/SlateRHIRenderer/Private/SlateRHIRenderingPolicy.h
@@ -80,6 +80,8 @@ public:
 
 	virtual void FlushGeneratedResources();
 
+	void BlurRectExternal(FRHICommandListImmediate& RHICmdList, FTexture2DRHIRef BlurSrc, FTextureReferenceRHIRef& BlurDst, FIntPoint DstExtent, float BlurStrength) const;
+
 private:
 	ETextureSamplerFilter GetSamplerFilter(const UTexture* Texture) const;
 
diff --git a/Engine/Source/Runtime/SlateRHIRenderer/Private/SlateRHIRenderingPolicyInterface.cpp b/Engine/Source/Runtime/SlateRHIRenderer/Private/SlateRHIRenderingPolicyInterface.cpp
new file mode 100644
--- /dev/null
+++ b/Engine/Source/Runtime/SlateRHIRenderer/Private/SlateRHIRenderingPolicyInterface.cpp
@@ -0,0 +1,42 @@
+// Copyright Epic Games, Inc. All Rights Reserved.
+
+#include "Interfaces/SlateRHIRenderingPolicyInterface.h"
+#include "SlateRHIRenderingPolicy.h"
+
+FSlateRHIRenderingPolicyInterface::FSlateRHIRenderingPolicyInterface(FSlateRHIRenderingPolicy* InRenderingPolicy)
+	: RenderingPolicy(InRenderingPolicy)
+{
+}
+
+bool FSlateRHIRenderingPolicyInterface::IsValid() const
+{
+	return RenderingPolicy != nullptr;
+}
+
+bool FSlateRHIRenderingPolicyInterface::IsVertexColorInLinearSpace() const
+{
+	if (RenderingPolicy)
+	{
+		return RenderingPolicy->IsVertexColorInLinearSpace();
+	}
+
+	return false;
+}
+
+bool FSlateRHIRenderingPolicyInterface::GetApplyColorDeficiencyCorrection() const
+{
+	if (RenderingPolicy)
+	{
+		return RenderingPolicy->GetApplyColorDeficiencyCorrection();
+	}
+
+	return false;
+}
+
+void FSlateRHIRenderingPolicyInterface::BlurRectExternal(FRHICommandListImmediate& RHICmdList, FTexture2DRHIRef BlurSrc, FTextureReferenceRHIRef& BlurDst, FIntPoint DstExtent, float BlurStrength) const
+{
+	if (RenderingPolicy)
+	{
+		RenderingPolicy->BlurRectExternal(RHICmdList, BlurSrc, BlurDst, DstExtent, BlurStrength);
+	}
+}
\ No newline at end of file
diff --git a/Engine/Source/Runtime/SlateRHIRenderer/Private/SlateUTextureResource.cpp b/Engine/Source/Runtime/SlateRHIRenderer/Private/SlateUTextureResource.cpp
--- a/Engine/Source/Runtime/SlateRHIRenderer/Private/SlateUTextureResource.cpp
+++ b/Engine/Source/Runtime/SlateRHIRenderer/Private/SlateUTextureResource.cpp
@@ -111,7 +111,7 @@ void FSlateUTextureResource::UpdateTexture(UTexture* InTexture)
 	if (Proxy && TextureObject)
 	{
 		CachedSlatePostBuffers = ESlatePostRT::None;
-		for (const TPair<ESlatePostRT, FSlatePostSettings>& SlatePostSetting : USlateRendererSettings::Get()->SlatePostSettings)
+		for (const TPair<ESlatePostRT, FSlatePostSettings>& SlatePostSetting : USlateRendererSettings::Get()->GetSlatePostSettings())
 		{
 			const ESlatePostRT SlatePostBitflag = SlatePostSetting.Key;
 			const FSlatePostSettings& SlatePostSettingValue = SlatePostSetting.Value;
diff --git a/Engine/Source/Runtime/SlateRHIRenderer/Public/FX/SlatePostBufferBlur.h b/Engine/Source/Runtime/SlateRHIRenderer/Public/FX/SlatePostBufferBlur.h
new file mode 100644
--- /dev/null
+++ b/Engine/Source/Runtime/SlateRHIRenderer/Public/FX/SlatePostBufferBlur.h
@@ -0,0 +1,28 @@
+// Copyright Epic Games, Inc. All Rights Reserved.
+
+#pragma once
+
+#include "CoreMinimal.h"
+#include "FX/SlateRHIPostBufferProcessor.h"
+
+#include "SlatePostBufferBlur.generated.h"
+
+/**
+ * Slate Post Buffer Processor that performs a simple gaussian blur to the backbuffer
+ * 
+ * Create a new asset deriving from this class to use / modify settings.
+ */
+UCLASS(Abstract, Blueprintable, CollapseCategories)
+class SLATERHIRENDERER_API USlatePostBufferBlur : public USlateRHIPostBufferProcessor
+{
+	GENERATED_BODY()
+
+public:
+
+	UPROPERTY(interp, BlueprintReadWrite, Category = "GaussianBlur")
+	float GaussianBlurStrength = 10;
+
+public:
+
+	virtual void PostProcess(FRenderResource* InViewInfo, FRenderResource* InViewportTexture, FVector2D InElementWindowSize, FSlateRHIRenderingPolicyInterface InRenderingPolicy, UTextureRenderTarget2D* InSlatePostBuffer);
+};
\ No newline at end of file
diff --git a/Engine/Source/Runtime/SlateRHIRenderer/Public/FX/SlateRHIPostBufferProcessor.h b/Engine/Source/Runtime/SlateRHIRenderer/Public/FX/SlateRHIPostBufferProcessor.h
new file mode 100644
--- /dev/null
+++ b/Engine/Source/Runtime/SlateRHIRenderer/Public/FX/SlateRHIPostBufferProcessor.h
@@ -0,0 +1,90 @@
+// Copyright Epic Games, Inc. All Rights Reserved.
+
+#pragma once
+
+#include "CoreMinimal.h"
+#include "Engine/TextureRenderTarget2D.h"
+#include "RenderResource.h"
+#include "RendererInterface.h"
+#include "Rendering/RenderingCommon.h"
+#include "Rendering/SlateRenderer.h"
+#include "Rendering/SlateRendererSettings.h"
+#include "RHIFwd.h"
+#include "Interfaces/SlateRHIRenderingPolicyInterface.h"
+
+#include "SlateRHIPostBufferProcessor.generated.h"
+
+/**
+ * Base class for types that can process the backbuffer scene into the slate post buffer.
+ */
+UCLASS(Abstract, Blueprintable, CollapseCategories)
+class SLATERHIRENDERER_API USlateRHIPostBufferProcessor : public USlateCorePostBufferProcessor
+{
+	GENERATED_BODY()
+
+public:
+
+	/**
+	 * Overridable postprocess for the given source scene backbuffer provided in 'Src' into 'Dst'
+	 * You must override this method if your postprocess has any parameters that can change at runtime.
+	 * In your override, you should copy those params before executing 'ENQUEUE_RENDER_COMMAND'.
+	 * This allows you to avoid render & game thread race conditions. See 'SlatePostBufferBlur' for example
+	 *
+	 * @param InViewInfo			'FViewportInfo' resource used to get backbuffer in standalone
+	 * @param InViewportTexture		'FSlateRenderTargetRHI' resource used to get the 'BufferedRT' viewport texture used in PIE
+	 * @param InElementWindowSize	Size of window being rendered, used to determine if using stereo rendering or not.
+	 * @param InRenderingPolicy		Slate RHI RenderingPolicy
+	 * @param InSlatePostBuffer		Texture render target used for final output
+	 */
+	virtual void PostProcess(FRenderResource* InViewInfo, FRenderResource* InViewportTexture, FVector2D InElementWindowSize, FSlateRHIRenderingPolicyInterface InRenderingPolicy, UTextureRenderTarget2D* InSlatePostBuffer);
+
+protected:
+
+	/**
+	 * Overridable postprocess for the given source scene backbuffer provided in 'Src' into 'Dst'
+	 * Override this method if you have no parameters that are mutable, otherwise you should 
+	 * override the above method & copy params (You can still use this interface if desired).
+	 * 
+	 * @param RenderingPolicy	Slate RHI RenderingPolicy
+	 * @param RHICmdList		RHI command list to queue commands on
+	 * @param Src				Source texture to process, currently should only be the scene backbuffer
+	 * @param Dst				Dst texture to output postprocess results to
+	 * @param DstExtent			Target extents / dimensions of 'Dst' after processing, use this since Dst's init may not be completed via RHI yet
+	 */
+	virtual void PostProcess_RenderThread(FSlateRHIRenderingPolicyInterface RenderingPolicy, FRHICommandListImmediate& RHICmdList, FTexture2DRHIRef Src, FTextureReferenceRHIRef& Dst, FIntPoint DstExtent)
+	{
+	};
+
+	/** 
+	 * Gets scene backbuffer, typically used as 'Src' texture for post process, but not always (Ex: PIE).
+	 * 
+	 * @param InViewInfo			'FViewportInfo' resource used to get backbuffer in standalone
+	 * @param InViewportTexture		'FSlateRenderTargetRHI' resource used to get the 'BufferedRT' viewport texture used in PIE
+	 * @param InElementWindowSize	Size of window being rendered, used to determine if using stereo rendering or not.
+	 * @param InRHICmdList			RHI command list to queue commands on
+	 */
+	FTexture2DRHIRef GetBackbuffer_RenderThread(FRenderResource* InViewInfo, FRenderResource* InViewportTexture, FVector2D InElementWindowSize, FRHICommandListImmediate& InRHICmdList);
+
+	/**
+	 * Gets 'Src' texture for post process command. Typically the scenebuffer.
+	 *
+	 * @param InBackBuffer			Backbuffer used in standalone
+	 * @param InViewportTexture		'FSlateRenderTargetRHI' resource used for 'BufferedRT' viewport texture in PIE
+	 */
+	FTexture2DRHIRef GetSrcTexture_RenderThread(FTexture2DRHIRef InBackBuffer, FRenderResource* InViewportTexture);
+
+	/**
+	 * Gets 'Dst' texture for post process command. Convience method, this should be possible through the direct resource.
+	 *
+	 * @param InSlatePostBuffer		Texture render target used for final output
+	 */
+	FTextureReferenceRHIRef& GetDstTexture_RenderThread(UTextureRenderTarget2D* InSlatePostBuffer);
+
+	/**
+	 * Gets 'Dst' extent. Used for final size in post process command.
+	 *
+	 * @param InBackBuffer			Backbuffer used for size in standalone
+	 * @param InViewportTexture		'FSlateRenderTargetRHI' resource used for 'BufferedRT' viewport texture size in PIE
+	 */
+	FIntPoint GetDstExtent_RenderThread(FTexture2DRHIRef InBackBuffer, FRenderResource* InViewportTexture);
+};
\ No newline at end of file
diff --git a/Engine/Source/Runtime/SlateRHIRenderer/Public/Interfaces/SlateRHIRenderingPolicyInterface.h b/Engine/Source/Runtime/SlateRHIRenderer/Public/Interfaces/SlateRHIRenderingPolicyInterface.h
new file mode 100644
--- /dev/null
+++ b/Engine/Source/Runtime/SlateRHIRenderer/Public/Interfaces/SlateRHIRenderingPolicyInterface.h
@@ -0,0 +1,33 @@
+// Copyright Epic Games, Inc. All Rights Reserved.
+
+#pragma once
+
+#include "CoreMinimal.h"
+#include "RendererInterface.h"
+#include "Rendering/RenderingCommon.h"
+#include "Rendering/ShaderResourceManager.h"
+#include "Rendering/DrawElements.h"
+#include "Rendering/RenderingPolicy.h"
+
+class FSlateRHIRenderingPolicy;
+
+/**
+ * Class used to expose only limited parts of the FSlateRenderingPolicy. 
+ * Not an interface to be implemented.
+ */
+class SLATERHIRENDERER_API FSlateRHIRenderingPolicyInterface
+{
+public:
+	FSlateRHIRenderingPolicyInterface(FSlateRHIRenderingPolicy* InRenderingPolicy);
+
+	bool IsValid() const;
+	bool IsVertexColorInLinearSpace() const;
+	bool GetApplyColorDeficiencyCorrection() const;
+
+	void BlurRectExternal(FRHICommandListImmediate& RHICmdList, FTexture2DRHIRef BlurSrc, FTextureReferenceRHIRef& BlurDst, FIntPoint DstExtent, float BlurStrength) const;
+
+private:
+
+	/** Rendering policy we are an interface to */
+	FSlateRHIRenderingPolicy* RenderingPolicy;
+};
