Fix backbuffers being released in D3D12 while they're still in use.

We need to sync the GPU before releasing viewport objects, because the release requests are placed on the deferred deletion queue, which is not synchronized with Present calls.

#rnx

diff --git a/Engine/Source/Runtime/SlateRHIRenderer/Private/SlateRHIRenderer.cpp b/Engine/Source/Runtime/SlateRHIRenderer/Private/SlateRHIRenderer.cpp
--- a/Engine/Source/Runtime/SlateRHIRenderer/Private/SlateRHIRenderer.cpp
+++ b/Engine/Source/Runtime/SlateRHIRenderer/Private/SlateRHIRenderer.cpp
@@ -154,6 +154,9 @@ void FViewportInfo::InitRHI(FRHICommandListBase&)
 
 void FViewportInfo::ReleaseRHI()
 {
+	// Some RHIs delete resources when the command list completes on the GPU, but they don't take Present calls into account due to
+	// API limitations. Sync the GPU here, to make sure nothing is using the backbuffer anymore.
+	FRHICommandListExecutor::GetImmediateCommandList().BlockUntilGPUIdle();
 	DepthStencil.SafeRelease();
 	ViewportRHI.SafeRelease();
 }
@@ -603,8 +606,7 @@ void FSlateRHIRenderer::OnWindowDestroyed(const TSharedRef<SWindow>& InWindow)
 
 		BeginReleaseResource(*ViewportInfoPtr);
 
-		// Need to flush rendering commands as the viewport may be in use by the render thread
-		// and the rendering resources must be released on the render thread before the viewport can be deleted
+		// Flush rendering commands again so that the resource deletion request is processed.
 		FlushRenderingCommands();
 
 		delete *ViewportInfoPtr;
