Thread safety / lifetime improvements to Slate PostBuffer

diff --git a/Engine/Source/Runtime/SlateCore/Private/Rendering/SlateRendererSettings.cpp b/Engine/Source/Runtime/SlateCore/Private/Rendering/SlateRendererSettings.cpp
deleted file mode 100644
diff --git a/Engine/Source/Runtime/SlateCore/Public/Rendering/ElementBatcher.h b/Engine/Source/Runtime/SlateCore/Public/Rendering/ElementBatcher.h
--- a/Engine/Source/Runtime/SlateCore/Public/Rendering/ElementBatcher.h
+++ b/Engine/Source/Runtime/SlateCore/Public/Rendering/ElementBatcher.h
@@ -4,7 +4,7 @@
 
 #include "CoreMinimal.h"
 #include "Rendering/RenderingCommon.h"
-#include "Rendering/SlateRendererSettings.h"
+#include "Rendering/SlateRendererTypes.h"
 #include "Layout/Clipping.h"
 #include "Stats/Stats.h"
 #include "SlateGlobals.h"
diff --git a/Engine/Source/Runtime/SlateCore/Public/Rendering/SlateRendererSettings.h b/Engine/Source/Runtime/SlateCore/Public/Rendering/SlateRendererSettings.h
deleted file mode 100644
diff --git a/Engine/Source/Runtime/SlateCore/Public/Rendering/SlateRendererTypes.h b/Engine/Source/Runtime/SlateCore/Public/Rendering/SlateRendererTypes.h
new file mode 100644
--- /dev/null
+++ b/Engine/Source/Runtime/SlateCore/Public/Rendering/SlateRendererTypes.h
@@ -0,0 +1,27 @@
+// Copyright Epic Games, Inc. All Rights Reserved.
+
+#pragma once
+
+#include "CoreMinimal.h"
+#include "Misc/EnumRange.h"
+
+#include "SlateRendererTypes.generated.h"
+
+/**
+ * Bitfield used to mark if a slate post RT is used or not
+ */
+UENUM(BlueprintType)
+enum class ESlatePostRT : uint8
+{
+	None = 0 << 0,
+	ESlatePostRT_0 = 1 << 0,
+	ESlatePostRT_1 = 1 << 1,
+	ESlatePostRT_2 = 1 << 2,
+	ESlatePostRT_3 = 1 << 3,
+	ESlatePostRT_4 = 1 << 4,
+	Num = 5
+};
+
+ENUM_CLASS_FLAGS(ESlatePostRT);
+
+ENUM_RANGE_BY_VALUES(ESlatePostRT, ESlatePostRT::ESlatePostRT_0, ESlatePostRT::ESlatePostRT_1, ESlatePostRT::ESlatePostRT_2, ESlatePostRT::ESlatePostRT_3, ESlatePostRT::ESlatePostRT_4);
diff --git a/Engine/Source/Runtime/SlateCore/Public/Textures/SlateShaderResource.h b/Engine/Source/Runtime/SlateCore/Public/Textures/SlateShaderResource.h
--- a/Engine/Source/Runtime/SlateCore/Public/Textures/SlateShaderResource.h
+++ b/Engine/Source/Runtime/SlateCore/Public/Textures/SlateShaderResource.h
@@ -5,7 +5,7 @@
 #include "CoreMinimal.h"
 #include "Stats/Stats.h"
 #include "SlateGlobals.h"
-#include "Rendering/SlateRendererSettings.h"
+#include "Rendering/SlateRendererTypes.h"
 #include "Rendering/SlateResourceHandle.h"
 
 class FSlateShaderResourceProxy;
diff --git a/Engine/Source/Runtime/SlateRHIRenderer/Private/FX/SlateFXSubsystem.cpp b/Engine/Source/Runtime/SlateRHIRenderer/Private/FX/SlateFXSubsystem.cpp
new file mode 100644
--- /dev/null
+++ b/Engine/Source/Runtime/SlateRHIRenderer/Private/FX/SlateFXSubsystem.cpp
@@ -0,0 +1,67 @@
+// Copyright Epic Games, Inc. All Rights Reserved.
+
+#include "FX/SlateFXSubsystem.h"
+
+#include "Engine/Engine.h"
+#include "Engine/GameInstance.h"
+
+#include UE_INLINE_GENERATED_CPP_BY_NAME(SlateFXSubsystem)
+
+bool USlateFXSubsystem::ShouldCreateSubsystem(UObject* Outer) const
+{
+	TArray<UClass*> ChildClasses;
+	GetDerivedClasses(GetClass(), ChildClasses, false);
+
+	// Only create an instance if there is no override implementation defined elsewhere
+	return ChildClasses.Num() == 0;
+}
+
+void USlateFXSubsystem::Initialize(FSubsystemCollectionBase& Collection)
+{
+	Super::Initialize(Collection);
+
+	FWorldDelegates::OnPreWorldInitialization.AddUObject(this, &USlateFXSubsystem::OnPreWorldInitialization);
+	FWorldDelegates::OnPostWorldCleanup.AddUObject(this, &USlateFXSubsystem::OnPostWorldCleanup);
+}
+
+void USlateFXSubsystem::Deinitialize()
+{
+	SlatePostBufferProcessors.Empty();
+
+	FWorldDelegates::OnPreWorldInitialization.RemoveAll(this);
+	FWorldDelegates::OnPostWorldCleanup.RemoveAll(this);
+
+	Super::Deinitialize();
+}
+
+USlateRHIPostBufferProcessor* USlateFXSubsystem::GetSlatePostProcessor(ESlatePostRT InPostBufferBit)
+{
+	if (TObjectPtr<USlateRHIPostBufferProcessor>* Processor = SlatePostBufferProcessors.Find(InPostBufferBit))
+	{
+		return *Processor;
+	}
+
+	return nullptr;
+}
+
+void USlateFXSubsystem::OnPreWorldInitialization(UWorld* World, const UWorld::InitializationValues IVS)
+{
+	SlatePostBufferProcessors.Empty();
+
+	if (const USlateRHIRendererSettings* SlateRendererSettings = USlateRHIRendererSettings::Get())
+	{
+		for (ESlatePostRT SlatePostBufferBit : TEnumRange<ESlatePostRT>())
+		{
+			const FSlatePostSettings& PostSetting = SlateRendererSettings->GetSlatePostSetting(SlatePostBufferBit);
+			if (PostSetting.bEnabled && PostSetting.PostProcessorClass)
+			{
+				SlatePostBufferProcessors.Add(SlatePostBufferBit, NewObject<USlateRHIPostBufferProcessor>(this, PostSetting.PostProcessorClass));
+			}
+		}
+	}
+}
+
+void USlateFXSubsystem::OnPostWorldCleanup(UWorld* World, bool SessionEnded, bool bCleanupResources)
+{
+	SlatePostBufferProcessors.Empty();
+}
diff --git a/Engine/Source/Runtime/SlateRHIRenderer/Private/FX/SlatePostBufferBlur.cpp b/Engine/Source/Runtime/SlateRHIRenderer/Private/FX/SlatePostBufferBlur.cpp
--- a/Engine/Source/Runtime/SlateRHIRenderer/Private/FX/SlatePostBufferBlur.cpp
+++ b/Engine/Source/Runtime/SlateRHIRenderer/Private/FX/SlatePostBufferBlur.cpp
@@ -14,20 +14,17 @@ void USlatePostBufferBlur::PostProcess(FRenderResource* InViewInfo, FRenderResou
 		// Explicit param copy to avoid renderthread from reading value during gamethread write
 		float GaussianBlurStrengthCopy = GaussianBlurStrength;
 
-		ENQUEUE_RENDER_COMMAND(FUpdateSlatePostBuffersWithFX_Blur)([Self = this, InViewInfo, InViewportTexture, InElementWindowSize, InRenderingPolicy, InSlatePostBuffer, GaussianBlurStrengthCopy](FRHICommandListImmediate& RHICmdList)
+		ENQUEUE_RENDER_COMMAND(FUpdateSlatePostBuffersWithFX_Blur)([InViewInfo, InViewportTexture, InElementWindowSize, InRenderingPolicy, InSlatePostBuffer, GaussianBlurStrengthCopy](FRHICommandListImmediate& RHICmdList)
 		{
-			if (Self)
-			{
-				FTexture2DRHIRef BackBuffer = Self->GetBackbuffer_RenderThread(InViewInfo, InViewportTexture, InElementWindowSize, RHICmdList);
+			FTexture2DRHIRef BackBuffer = USlateRHIPostBufferProcessor::GetBackbuffer_RenderThread(InViewInfo, InViewportTexture, InElementWindowSize, RHICmdList);
 
-				if (BackBuffer)
-				{
-					FTexture2DRHIRef Src = Self->GetSrcTexture_RenderThread(BackBuffer, InViewportTexture);
-					FTextureReferenceRHIRef& Dst = Self->GetDstTexture_RenderThread(InSlatePostBuffer);
-					FIntPoint DstExtent = Self->GetDstExtent_RenderThread(BackBuffer, InViewportTexture);
+			if (BackBuffer)
+			{
+				FTexture2DRHIRef Src = USlateRHIPostBufferProcessor::GetSrcTexture_RenderThread(BackBuffer, InViewportTexture);
+				FTextureReferenceRHIRef& Dst = USlateRHIPostBufferProcessor::GetDstTexture_RenderThread(InSlatePostBuffer);
+				FIntPoint DstExtent = USlateRHIPostBufferProcessor::GetDstExtent_RenderThread(BackBuffer, InViewportTexture);
 
-					InRenderingPolicy.BlurRectExternal(RHICmdList, Src, Dst, DstExtent, GaussianBlurStrengthCopy);
-				}
+				InRenderingPolicy.BlurRectExternal(RHICmdList, Src, Dst, DstExtent, GaussianBlurStrengthCopy);
 			}
 		});
 	}
diff --git a/Engine/Source/Runtime/SlateRHIRenderer/Private/FX/SlateRHIPostBufferProcessor.cpp b/Engine/Source/Runtime/SlateRHIRenderer/Private/FX/SlateRHIPostBufferProcessor.cpp
--- a/Engine/Source/Runtime/SlateRHIRenderer/Private/FX/SlateRHIPostBufferProcessor.cpp
+++ b/Engine/Source/Runtime/SlateRHIRenderer/Private/FX/SlateRHIPostBufferProcessor.cpp
@@ -9,26 +9,6 @@
 
 #include UE_INLINE_GENERATED_CPP_BY_NAME(SlateRHIPostBufferProcessor)
 
-void USlateRHIPostBufferProcessor::PostProcess(FRenderResource* InViewInfo, FRenderResource* InViewportTexture, FVector2D InElementWindowSize, FSlateRHIRenderingPolicyInterface InRenderingPolicy, UTextureRenderTarget2D* InSlatePostBuffer)
-{
-	ENQUEUE_RENDER_COMMAND(FUpdateSlatePostBuffersWithFX)([Self = this, InViewInfo, InViewportTexture, InElementWindowSize, InRenderingPolicy, InSlatePostBuffer](FRHICommandListImmediate& RHICmdList)
-	{
-		if (Self)
-		{
-			FTexture2DRHIRef BackBuffer = Self->GetBackbuffer_RenderThread(InViewInfo, InViewportTexture, InElementWindowSize, RHICmdList);
-
-			if (BackBuffer)
-			{
-				FTexture2DRHIRef Src = Self->GetSrcTexture_RenderThread(BackBuffer, InViewportTexture);
-				FTextureReferenceRHIRef& Dst = Self->GetDstTexture_RenderThread(InSlatePostBuffer);
-				FIntPoint DstExtent = Self->GetDstExtent_RenderThread(BackBuffer, InViewportTexture);
-
-				Self->PostProcess_RenderThread(InRenderingPolicy, RHICmdList, Src, Dst, DstExtent);
-			}
-		}
-	});
-}
-
 FTexture2DRHIRef USlateRHIPostBufferProcessor::GetBackbuffer_RenderThread(FRenderResource* InViewInfo, FRenderResource* InViewportTexture, FVector2D InElementWindowSize, FRHICommandListImmediate& InRHICmdList)
 {
 	FViewportInfo* ViewInfo = static_cast<FViewportInfo*>(InViewInfo);
diff --git a/Engine/Source/Runtime/SlateRHIRenderer/Private/SlateMaterialResource.cpp b/Engine/Source/Runtime/SlateRHIRenderer/Private/SlateMaterialResource.cpp
--- a/Engine/Source/Runtime/SlateRHIRenderer/Private/SlateMaterialResource.cpp
+++ b/Engine/Source/Runtime/SlateRHIRenderer/Private/SlateMaterialResource.cpp
@@ -1,7 +1,7 @@
 // Copyright Epic Games, Inc. All Rights Reserved.
 
 #include "SlateMaterialResource.h"
-#include "Rendering/SlateRendererSettings.h"
+#include "SlateRHIRendererSettings.h"
 #include "Engine/Texture.h"
 #include "Materials/MaterialInstanceDynamic.h"
 #include "Materials/MaterialRenderProxy.h"
@@ -148,7 +148,7 @@ void FSlateMaterialResource::UpdateMaterial(const UMaterialInterface& InMaterial
 		CachedSlatePostBuffers = ESlatePostRT::None;
 		for (const UTexture* OutUsedTexture : OutUsedTextures)
 		{
-			for (const TPair<ESlatePostRT, FSlatePostSettings>& SlatePostSetting : USlateRendererSettings::Get()->GetSlatePostSettings())
+			for (const TPair<ESlatePostRT, FSlatePostSettings>& SlatePostSetting : USlateRHIRendererSettings::Get()->GetSlatePostSettings())
 			{
 				const ESlatePostRT SlatePostBitflag = SlatePostSetting.Key;
 				const FSlatePostSettings& SlatePostSettingValue = SlatePostSetting.Value;
diff --git a/Engine/Source/Runtime/SlateRHIRenderer/Private/SlateRHIRenderer.cpp b/Engine/Source/Runtime/SlateRHIRenderer/Private/SlateRHIRenderer.cpp
--- a/Engine/Source/Runtime/SlateRHIRenderer/Private/SlateRHIRenderer.cpp
+++ b/Engine/Source/Runtime/SlateRHIRenderer/Private/SlateRHIRenderer.cpp
@@ -3,6 +3,7 @@
 #include "SlateRHIRenderer.h"
 #include "Fonts/FontCache.h"
 #include "SlateRHIRenderingPolicy.h"
+#include "SlateRHIRendererSettings.h"
 #include "Misc/ScopeLock.h"
 #include "Modules/ModuleManager.h"
 #include "Styling/CoreStyle.h"
@@ -11,6 +12,7 @@
 #include "EngineGlobals.h"
 #include "Engine/AssetManager.h"
 #include "Engine/TextureRenderTarget2D.h"
+#include "FX/SlateFXSubsystem.h"
 #include "FX/SlateRHIPostBufferProcessor.h"
 #include "Materials/MaterialRenderProxy.h"
 #include "MaterialShared.h"
@@ -25,7 +27,6 @@
 #include "SlateShaders.h"
 #include "Rendering/ElementBatcher.h"
 #include "Rendering/SlateRenderer.h"
-#include "Rendering/SlateRendererSettings.h"
 #include "RenderResource.h"
 #include "RenderingThread.h"
 #include "RHIResources.h"
@@ -1751,16 +1752,16 @@ void FSlateRHIRenderer::DrawWindows_Private(FSlateDrawBuffer& WindowDrawBuffer)
 					uint8 SlatePostBufferBitIndex = 0;
 					for (ESlatePostRT SlatePostBufferBit : TEnumRange<ESlatePostRT>())
 					{
-						if (!USlateRendererSettings::Get()->GetSlatePostSetting(SlatePostBufferBit).bEnabled)
+						if (!USlateRHIRendererSettings::Get()->GetSlatePostSetting(SlatePostBufferBit).bEnabled)
 						{
 							SlatePostBufferBitIndex++;
 							continue;
 						}
 
-						UTextureRenderTarget2D* SlatePostBuffer = Cast<UTextureRenderTarget2D>(USlateRendererSettings::Get()->TryGetPostBufferRT(SlatePostBufferBit));
+						UTextureRenderTarget2D* SlatePostBuffer = Cast<UTextureRenderTarget2D>(USlateRHIRendererSettings::Get()->TryGetPostBufferRT(SlatePostBufferBit));
 						if (!SlatePostBuffer)
 						{
-							SlatePostBuffer = Cast<UTextureRenderTarget2D>(USlateRendererSettings::GetMutable()->LoadGetPostBufferRT(SlatePostBufferBit));
+							SlatePostBuffer = Cast<UTextureRenderTarget2D>(USlateRHIRendererSettings::GetMutable()->LoadGetPostBufferRT(SlatePostBufferBit));
 						}
 
 						bool bIsViewportPresentForPIE = GIsEditor ? Window->GetViewport().IsValid() : true;
@@ -1776,14 +1777,14 @@ void FSlateRHIRenderer::DrawWindows_Private(FSlateDrawBuffer& WindowDrawBuffer)
 							{
 								ViewportTexture = static_cast<FSlateRenderTargetRHI*>(Window->GetViewport()->GetViewportRenderTargetTexture());
 							}
-							
+
 							FIntPoint SizeSlatePostRT = GIsEditor
 								? FIntPoint(ViewportTexture->GetWidth(), ViewportTexture->GetHeight())
 								: FIntPoint(ViewInfo->DesiredWidth, ViewInfo->DesiredHeight);
 
 							bool bHDREnabled = IsHDREnabled();
-							bool bIsPixelFormatCorrect = bHDREnabled 
-								? SlatePostBuffer->GetFormat() == EPixelFormat::PF_FloatRGBA 
+							bool bIsPixelFormatCorrect = bHDREnabled
+								? SlatePostBuffer->GetFormat() == EPixelFormat::PF_FloatRGBA
 								: SlatePostBuffer->GetFormat() == EPixelFormat::PF_A2B10G10R10;
 							if (SlatePostBuffer->SizeX != SizeSlatePostRT.X || SlatePostBuffer->SizeY != SizeSlatePostRT.Y || !bIsPixelFormatCorrect)
 							{
@@ -1792,7 +1793,22 @@ void FSlateRHIRenderer::DrawWindows_Private(FSlateDrawBuffer& WindowDrawBuffer)
 
 							const FVector2D ElementWindowSize = ElementList.GetWindowSize();
 
-							if (USlateRHIPostBufferProcessor* PostProcessor = Cast<USlateRHIPostBufferProcessor>(USlateRendererSettings::Get()->GetSlatePostProcessor(SlatePostBufferBit)))
+							auto GetPostProcessor = [&](ESlatePostRT InSlatePostBufferBit)
+							{
+								USlateRHIPostBufferProcessor* Result = nullptr;
+
+								if (GEngine)
+								{
+									if (USlateFXSubsystem* SlateFXSubsystem = GEngine->GetEngineSubsystem<USlateFXSubsystem>())
+									{
+										Result = SlateFXSubsystem->GetSlatePostProcessor(InSlatePostBufferBit);
+									}
+								}
+
+								return Result;
+							};
+
+							if (USlateRHIPostBufferProcessor* PostProcessor = GetPostProcessor(SlatePostBufferBit))
 							{
 								// Allow the post processor to enque render commands, we delgate this task so it can be done in a thread safe manner.
 								PostProcessor->PostProcess(ViewInfo, ViewportTexture, ElementWindowSize, FSlateRHIRenderingPolicyInterface(RenderingPolicy.Get()), SlatePostBuffer);
diff --git a/Engine/Source/Runtime/SlateRHIRenderer/Private/SlateRHIRenderer.h b/Engine/Source/Runtime/SlateRHIRenderer/Private/SlateRHIRenderer.h
--- a/Engine/Source/Runtime/SlateRHIRenderer/Private/SlateRHIRenderer.h
+++ b/Engine/Source/Runtime/SlateRHIRenderer/Private/SlateRHIRenderer.h
@@ -8,10 +8,10 @@
 #include "RHI.h"
 #include "RenderCommandFence.h"
 #include "RenderResource.h"
+#include "SlateRHIRendererSettings.h"
 #include "SlateRHIResourceManager.h"
 #include "UnrealClient.h"
 #include "Rendering/SlateRenderer.h"
-#include "Rendering/SlateRendererSettings.h"
 #include "Rendering/SlateDrawBuffer.h"
 #include "Slate/SlateTextures.h"
 #include "RendererInterface.h"
diff --git a/Engine/Source/Runtime/SlateRHIRenderer/Private/SlateRHIRendererSettings.cpp b/Engine/Source/Runtime/SlateRHIRenderer/Private/SlateRHIRendererSettings.cpp
new file mode 100644
--- /dev/null
+++ b/Engine/Source/Runtime/SlateRHIRenderer/Private/SlateRHIRendererSettings.cpp
@@ -0,0 +1,91 @@
+// Copyright Epic Games, Inc. All Rights Reserved.
+
+#include "SlateRHIRendererSettings.h"
+#include "HAL/IConsoleManager.h"
+
+#include UE_INLINE_GENERATED_CPP_BY_NAME(SlateRHIRendererSettings)
+
+static TAutoConsoleVariable<int32> CVarDefaultEnablePostRenderTarget_0(
+	TEXT("Slate.DefaultEnablePostRenderTarget_0"),
+	1,
+	TEXT("Experimental. Set true to enable slate post render target 0"),
+	ECVF_ReadOnly);
+
+FSlatePostSettings::FSlatePostSettings()
+	: bEnabled(false)
+	, PostProcessorClass(nullptr)
+	, PathToSlatePostRT(FString())
+	, CachedSlatePostRT(nullptr)
+{
+}
+
+USlateRHIRendererSettings::USlateRHIRendererSettings()
+{
+	SlatePostSettings.Add(ESlatePostRT::ESlatePostRT_0, FSlatePostSettings());
+	SlatePostSettings.Add(ESlatePostRT::ESlatePostRT_1, FSlatePostSettings());
+	SlatePostSettings.Add(ESlatePostRT::ESlatePostRT_2, FSlatePostSettings());
+	SlatePostSettings.Add(ESlatePostRT::ESlatePostRT_3, FSlatePostSettings());
+	SlatePostSettings.Add(ESlatePostRT::ESlatePostRT_4, FSlatePostSettings());
+
+	// By default, enable the first post RT
+	SlatePostSettings[ESlatePostRT::ESlatePostRT_0].bEnabled = CVarDefaultEnablePostRenderTarget_0.GetValueOnAnyThread();
+
+	// Hardcoded paths to engine assets
+	SlatePostSettings[ESlatePostRT::ESlatePostRT_0].PathToSlatePostRT = "/Engine/EngineResources/SlatePost0_RT.SlatePost0_RT";
+	SlatePostSettings[ESlatePostRT::ESlatePostRT_1].PathToSlatePostRT = "/Engine/EngineResources/SlatePost1_RT.SlatePost1_RT";
+	SlatePostSettings[ESlatePostRT::ESlatePostRT_2].PathToSlatePostRT = "/Engine/EngineResources/SlatePost2_RT.SlatePost2_RT";
+	SlatePostSettings[ESlatePostRT::ESlatePostRT_3].PathToSlatePostRT = "/Engine/EngineResources/SlatePost3_RT.SlatePost3_RT";
+	SlatePostSettings[ESlatePostRT::ESlatePostRT_4].PathToSlatePostRT = "/Engine/EngineResources/SlatePost4_RT.SlatePost4_RT";
+}
+
+USlateRHIRendererSettings::~USlateRHIRendererSettings()
+{
+	for (TPair<ESlatePostRT, FSlatePostSettings>& SlatePostSetting : SlatePostSettings)
+	{
+		FSlatePostSettings& PostSetting = SlatePostSetting.Value;
+
+		UObject* SlatePostBuffer = PostSetting.CachedSlatePostRT;
+		if (SlatePostBuffer)
+		{
+			SlatePostBuffer->RemoveFromRoot();
+		}
+	}
+}
+
+FSlatePostSettings& USlateRHIRendererSettings::GetMutableSlatePostSetting(ESlatePostRT InPostBufferBit)
+{
+	return SlatePostSettings[InPostBufferBit];
+}
+
+const FSlatePostSettings& USlateRHIRendererSettings::GetSlatePostSetting(ESlatePostRT InPostBufferBit) const
+{
+	return SlatePostSettings[InPostBufferBit];
+}
+
+UObject* USlateRHIRendererSettings::TryGetPostBufferRT(ESlatePostRT InPostBufferBit) const
+{
+	return SlatePostSettings[InPostBufferBit].CachedSlatePostRT;
+}
+
+UObject* USlateRHIRendererSettings::LoadGetPostBufferRT(ESlatePostRT InPostBufferBit)
+{
+	UObject* Result = TryGetPostBufferRT(InPostBufferBit);
+
+	if (!Result)
+	{
+		Result = LoadObject<UObject>(nullptr, *SlatePostSettings[InPostBufferBit].PathToSlatePostRT, nullptr, LOAD_None, nullptr);
+
+		if (Result)
+		{
+			Result->AddToRoot();
+			SlatePostSettings[InPostBufferBit].CachedSlatePostRT = Result;
+		}
+	}
+
+	return Result;
+}
+
+const TMap<ESlatePostRT, FSlatePostSettings>& USlateRHIRendererSettings::GetSlatePostSettings() const
+{
+	return SlatePostSettings;
+}
diff --git a/Engine/Source/Runtime/SlateRHIRenderer/Private/SlateUTextureResource.cpp b/Engine/Source/Runtime/SlateRHIRenderer/Private/SlateUTextureResource.cpp
--- a/Engine/Source/Runtime/SlateRHIRenderer/Private/SlateUTextureResource.cpp
+++ b/Engine/Source/Runtime/SlateRHIRenderer/Private/SlateUTextureResource.cpp
@@ -1,6 +1,7 @@
 // Copyright Epic Games, Inc. All Rights Reserved.
 
 #include "SlateUTextureResource.h"
+#include "SlateRHIRendererSettings.h"
 
 FSlateBaseUTextureResource::FSlateBaseUTextureResource(UTexture* InTexture)
 	: TextureObject(InTexture)
@@ -111,7 +112,7 @@ void FSlateUTextureResource::UpdateTexture(UTexture* InTexture)
 	if (Proxy && TextureObject)
 	{
 		CachedSlatePostBuffers = ESlatePostRT::None;
-		for (const TPair<ESlatePostRT, FSlatePostSettings>& SlatePostSetting : USlateRendererSettings::Get()->GetSlatePostSettings())
+		for (const TPair<ESlatePostRT, FSlatePostSettings>& SlatePostSetting : USlateRHIRendererSettings::Get()->GetSlatePostSettings())
 		{
 			const ESlatePostRT SlatePostBitflag = SlatePostSetting.Key;
 			const FSlatePostSettings& SlatePostSettingValue = SlatePostSetting.Value;
diff --git a/Engine/Source/Runtime/SlateRHIRenderer/Public/FX/SlateFXSubsystem.h b/Engine/Source/Runtime/SlateRHIRenderer/Public/FX/SlateFXSubsystem.h
new file mode 100644
--- /dev/null
+++ b/Engine/Source/Runtime/SlateRHIRenderer/Public/FX/SlateFXSubsystem.h
@@ -0,0 +1,44 @@
+// Copyright Epic Games, Inc. All Rights Reserved.
+
+#pragma once
+
+#include "Subsystems/EngineSubsystem.h"
+
+#include "Engine/World.h"
+#include "SlateRHIRendererSettings.h"
+
+#include "SlateFXSubsystem.generated.h"
+
+UCLASS(DisplayName = "Slate FX Subsystem")
+class SLATERHIRENDERER_API USlateFXSubsystem : public UEngineSubsystem
+{
+	GENERATED_BODY()
+
+public:
+
+	//~Begin UGameInstanceSubsystem Insterface
+	virtual bool ShouldCreateSubsystem(UObject* Outer) const override;
+	virtual void Initialize(FSubsystemCollectionBase& Collection) override;
+	virtual void Deinitialize() override;
+	//~End UGameInstanceSubsystem Insterface
+
+public:
+
+	/** Get post processor for a particular post buffer index, if it exists */
+	UFUNCTION(BlueprintCallable, Category = "SlateFX")
+	USlateRHIPostBufferProcessor* GetSlatePostProcessor(ESlatePostRT InPostBufferBit);
+
+private:
+
+	/** Map of post RT buffer index to buffer processors, if they exist */
+	UPROPERTY(Transient)
+	TMap<ESlatePostRT, TObjectPtr<USlateRHIPostBufferProcessor>> SlatePostBufferProcessors;
+
+private:
+
+	/** Callback to create processors on world init */
+	void OnPreWorldInitialization(UWorld* World, const UWorld::InitializationValues IVS);
+
+	/** Callback to remove processors on world cleanup */
+	void OnPostWorldCleanup(UWorld* World, bool SessionEnded, bool bCleanupResources);
+};
diff --git a/Engine/Source/Runtime/SlateRHIRenderer/Public/FX/SlateRHIPostBufferProcessor.h b/Engine/Source/Runtime/SlateRHIRenderer/Public/FX/SlateRHIPostBufferProcessor.h
--- a/Engine/Source/Runtime/SlateRHIRenderer/Public/FX/SlateRHIPostBufferProcessor.h
+++ b/Engine/Source/Runtime/SlateRHIRenderer/Public/FX/SlateRHIPostBufferProcessor.h
@@ -8,17 +8,19 @@
 #include "RendererInterface.h"
 #include "Rendering/RenderingCommon.h"
 #include "Rendering/SlateRenderer.h"
-#include "Rendering/SlateRendererSettings.h"
 #include "RHIFwd.h"
+#include "UObject/Object.h"
 #include "Interfaces/SlateRHIRenderingPolicyInterface.h"
 
 #include "SlateRHIPostBufferProcessor.generated.h"
 
 /**
  * Base class for types that can process the backbuffer scene into the slate post buffer.
+ * 
+ * Implement 'PostProcess' in your derived class. For an example see: USlatePostBufferBlur.
  */
 UCLASS(Abstract, Blueprintable, CollapseCategories)
-class SLATERHIRENDERER_API USlateRHIPostBufferProcessor : public USlateCorePostBufferProcessor
+class SLATERHIRENDERER_API USlateRHIPostBufferProcessor : public UObject
 {
 	GENERATED_BODY()
 
@@ -26,9 +28,9 @@ public:
 
 	/**
 	 * Overridable postprocess for the given source scene backbuffer provided in 'Src' into 'Dst'
-	 * You must override this method if your postprocess has any parameters that can change at runtime.
-	 * In your override, you should copy those params before executing 'ENQUEUE_RENDER_COMMAND'.
-	 * This allows you to avoid render & game thread race conditions. See 'SlatePostBufferBlur' for example
+	 * You must override this method. In your override, you should copy params before executing 'ENQUEUE_RENDER_COMMAND'.
+	 * This allows you to avoid render & game thread race conditions. See 'SlatePostBufferBlur' for example.
+	 * Also avoid capturing [this] in your override, to avoid possible GC issues with the processor instance.
 	 *
 	 * @param InViewInfo			'FViewportInfo' resource used to get backbuffer in standalone
 	 * @param InViewportTexture		'FSlateRenderTargetRHI' resource used to get the 'BufferedRT' viewport texture used in PIE
@@ -36,25 +38,12 @@ public:
 	 * @param InRenderingPolicy		Slate RHI RenderingPolicy
 	 * @param InSlatePostBuffer		Texture render target used for final output
 	 */
-	virtual void PostProcess(FRenderResource* InViewInfo, FRenderResource* InViewportTexture, FVector2D InElementWindowSize, FSlateRHIRenderingPolicyInterface InRenderingPolicy, UTextureRenderTarget2D* InSlatePostBuffer);
+	virtual void PostProcess(FRenderResource* InViewInfo, FRenderResource* InViewportTexture, FVector2D InElementWindowSize, FSlateRHIRenderingPolicyInterface InRenderingPolicy, UTextureRenderTarget2D* InSlatePostBuffer)
+	{
+	}
 
 protected:
 
-	/**
-	 * Overridable postprocess for the given source scene backbuffer provided in 'Src' into 'Dst'
-	 * Override this method if you have no parameters that are mutable, otherwise you should 
-	 * override the above method & copy params (You can still use this interface if desired).
-	 * 
-	 * @param RenderingPolicy	Slate RHI RenderingPolicy
-	 * @param RHICmdList		RHI command list to queue commands on
-	 * @param Src				Source texture to process, currently should only be the scene backbuffer
-	 * @param Dst				Dst texture to output postprocess results to
-	 * @param DstExtent			Target extents / dimensions of 'Dst' after processing, use this since Dst's init may not be completed via RHI yet
-	 */
-	virtual void PostProcess_RenderThread(FSlateRHIRenderingPolicyInterface RenderingPolicy, FRHICommandListImmediate& RHICmdList, FTexture2DRHIRef Src, FTextureReferenceRHIRef& Dst, FIntPoint DstExtent)
-	{
-	};
-
 	/** 
 	 * Gets scene backbuffer, typically used as 'Src' texture for post process, but not always (Ex: PIE).
 	 * 
@@ -63,7 +52,7 @@ protected:
 	 * @param InElementWindowSize	Size of window being rendered, used to determine if using stereo rendering or not.
 	 * @param InRHICmdList			RHI command list to queue commands on
 	 */
-	FTexture2DRHIRef GetBackbuffer_RenderThread(FRenderResource* InViewInfo, FRenderResource* InViewportTexture, FVector2D InElementWindowSize, FRHICommandListImmediate& InRHICmdList);
+	static FTexture2DRHIRef GetBackbuffer_RenderThread(FRenderResource* InViewInfo, FRenderResource* InViewportTexture, FVector2D InElementWindowSize, FRHICommandListImmediate& InRHICmdList);
 
 	/**
 	 * Gets 'Src' texture for post process command. Typically the scenebuffer.
@@ -71,14 +60,14 @@ protected:
 	 * @param InBackBuffer			Backbuffer used in standalone
 	 * @param InViewportTexture		'FSlateRenderTargetRHI' resource used for 'BufferedRT' viewport texture in PIE
 	 */
-	FTexture2DRHIRef GetSrcTexture_RenderThread(FTexture2DRHIRef InBackBuffer, FRenderResource* InViewportTexture);
+	static FTexture2DRHIRef GetSrcTexture_RenderThread(FTexture2DRHIRef InBackBuffer, FRenderResource* InViewportTexture);
 
 	/**
 	 * Gets 'Dst' texture for post process command. Convience method, this should be possible through the direct resource.
 	 *
 	 * @param InSlatePostBuffer		Texture render target used for final output
 	 */
-	FTextureReferenceRHIRef& GetDstTexture_RenderThread(UTextureRenderTarget2D* InSlatePostBuffer);
+	static FTextureReferenceRHIRef& GetDstTexture_RenderThread(UTextureRenderTarget2D* InSlatePostBuffer);
 
 	/**
 	 * Gets 'Dst' extent. Used for final size in post process command.
@@ -86,5 +75,5 @@ protected:
 	 * @param InBackBuffer			Backbuffer used for size in standalone
 	 * @param InViewportTexture		'FSlateRenderTargetRHI' resource used for 'BufferedRT' viewport texture size in PIE
 	 */
-	FIntPoint GetDstExtent_RenderThread(FTexture2DRHIRef InBackBuffer, FRenderResource* InViewportTexture);
+	static FIntPoint GetDstExtent_RenderThread(FTexture2DRHIRef InBackBuffer, FRenderResource* InViewportTexture);
 };
\ No newline at end of file
diff --git a/Engine/Source/Runtime/SlateRHIRenderer/Public/SlateRHIRendererSettings.h b/Engine/Source/Runtime/SlateRHIRenderer/Public/SlateRHIRendererSettings.h
new file mode 100644
--- /dev/null
+++ b/Engine/Source/Runtime/SlateRHIRenderer/Public/SlateRHIRendererSettings.h
@@ -0,0 +1,120 @@
+// Copyright Epic Games, Inc. All Rights Reserved.
+
+#pragma once
+
+#include "CoreMinimal.h"
+#include "Containers/ContainerAllocationPolicies.h"
+#include "Engine/DeveloperSettings.h"
+#include "FX/SlateRHIPostBufferProcessor.h"
+#include "Misc/EnumRange.h"
+#include "UObject/Object.h"
+#include "UObject/SoftObjectPath.h"
+#include "Rendering/SlateRendererTypes.h"
+#include "Templates/SubclassOf.h"
+
+#include "SlateRHIRendererSettings.generated.h"
+
+/**
+ * Settings for a particular Slate Post RT.
+ * Notably if enabled & blur by default. To be updated with additional effects & to be expandable in game code / settings.
+ */
+USTRUCT(BlueprintType, meta=(HiddenByDefault, DisableSplitPin))
+struct SLATERHIRENDERER_API FSlatePostSettings
+{
+	GENERATED_BODY()
+
+public:
+
+	FSlatePostSettings();
+
+	friend class USlateRHIRendererSettings;
+	friend class USlateFXSubsytem;
+
+public:
+
+	/** Should this post buffer be enabled for updating */
+	UPROPERTY(EditAnywhere, BlueprintReadWrite, Category=BufferSettings, meta=(PinHiddenByDefault))
+	uint8 bEnabled:1;
+
+	/** Copy of actually loaded post processor class */
+	UPROPERTY(EditAnywhere, BlueprintReadWrite, Category = BufferSettings, Meta = (AllowAbstract = false))
+	TSubclassOf<USlateRHIPostBufferProcessor> PostProcessorClass;
+
+private:
+
+	/** Path to Slate Post RT Asset */
+	UPROPERTY()
+	FString PathToSlatePostRT;
+	
+public:
+
+	/** Get asset name for post RT texture */
+	const FString& GetPathToSlatePostRT() const { return PathToSlatePostRT; }
+
+private:
+
+	/** Cached load of Slate Post RT Asset */
+	UPROPERTY(Transient)
+	TObjectPtr<UObject> CachedSlatePostRT;
+};
+
+/**
+ * Settings used to control slate rendering
+ */
+UCLASS(config = Game, defaultconfig)
+class SLATERHIRENDERER_API USlateRHIRendererSettings : public UDeveloperSettings
+{
+	GENERATED_BODY()
+
+public:
+		
+	static const USlateRHIRendererSettings* Get()
+	{ 
+		return GetDefault<USlateRHIRendererSettings>();
+	}
+
+	static USlateRHIRendererSettings* GetMutable()
+	{
+		return GetMutableDefault<USlateRHIRendererSettings>();
+	}
+
+public:
+
+	USlateRHIRendererSettings();
+	~USlateRHIRendererSettings();
+
+public:
+
+	/** Get settings struct for a particular post buffer index */
+	UFUNCTION(BlueprintCallable, Category = "SlateFX")
+	FSlatePostSettings& GetMutableSlatePostSetting(ESlatePostRT InPostBufferBit);
+
+	/** Get settings struct for a particular post buffer index */
+	UFUNCTION(BlueprintCallable, Category = "SlateFX")
+	const FSlatePostSettings& GetSlatePostSetting(ESlatePostRT InPostBufferBit) const;
+
+public:
+
+	/** Try to get post RT asset, returns nullptr if not already loaded */
+	UObject* TryGetPostBufferRT(ESlatePostRT InPostBufferBit) const;
+
+	/** Get post RT asset, loading if not already loaded */
+	UObject* LoadGetPostBufferRT(ESlatePostRT InPostBufferBit);
+
+	/** Get slate post settings map, non mutable */
+	const TMap<ESlatePostRT, FSlatePostSettings>& GetSlatePostSettings() const;
+
+private:
+
+	/** 
+	 * Map of all slate post RT's and their settings 
+	 * Note that each post RT used in a frame will result in 1 full framebuffer copy for slate to sample from.
+	 * If a post RT is not used, no copy will occur & that post RT will be resized to 1x1 after 2 frames of non-use.
+	 * 
+	 * By default only SlatePostRT_0 is enabled. The rest must manually be enabled in settings below.
+	 */
+
+	// Map is nice since needs no editor customization. After initial run there should be no more than 5 lookups each frame.
+	UPROPERTY(config, EditAnywhere, EditFixedSize, Category = "PostProcessing")
+	TMap<ESlatePostRT, FSlatePostSettings> SlatePostSettings;
+};
