Add support for layered Slate Post RT processing
---
Add PostBufferUpdate widget which can be used to trigger an update to the Slate Post RT when it "draws"
Updated Post RT will contain Scene + All UI drawn at that time & also be processed by provided processor.
Updated ICustomSlateElement interface to support draw context params.

diff --git a/Engine/Plugins/Experimental/Mutable/Source/CustomizableObjectEditor/Private/MuCOE/SCustomizableObjectLayoutGrid.cpp b/Engine/Plugins/Experimental/Mutable/Source/CustomizableObjectEditor/Private/MuCOE/SCustomizableObjectLayoutGrid.cpp
--- a/Engine/Plugins/Experimental/Mutable/Source/CustomizableObjectEditor/Private/MuCOE/SCustomizableObjectLayoutGrid.cpp
+++ b/Engine/Plugins/Experimental/Mutable/Source/CustomizableObjectEditor/Private/MuCOE/SCustomizableObjectLayoutGrid.cpp
@@ -97,7 +97,7 @@ public:
 
 private:
 
-	void DrawRenderThread(FRHICommandListImmediate& RHICmdList, const void* RenderTarget) override;
+	virtual void Draw_RenderThread(class FRHICommandListImmediate& RHICmdList, const void* InWindowBackBuffer, const FSlateCustomDrawParams& Params) override;
 
 	/** Basic function to draw a block in the canvas */
 	void DrawBlock(FBatchedElements* BatchedElements, const FHitProxyId HitProxyId, const FRect2D& BlockRect, FColor Color);
@@ -1219,7 +1219,7 @@ void FUVCanvasDrawer::SetLayoutMode(ELayoutGridMode Mode)
 }
 
 
-void FUVCanvasDrawer::DrawRenderThread(class FRHICommandListImmediate& RHICmdList, const void* InWindowBackBuffer)
+void FUVCanvasDrawer::Draw_RenderThread(class FRHICommandListImmediate& RHICmdList, const void* InWindowBackBuffer, const FSlateCustomDrawParams& Params)
 {
 	if (Initialized)
 	{
diff --git a/Engine/Source/Editor/GraphEditor/Private/MaterialNodes/SGraphNodeMaterialBase.cpp b/Engine/Source/Editor/GraphEditor/Private/MaterialNodes/SGraphNodeMaterialBase.cpp
--- a/Engine/Source/Editor/GraphEditor/Private/MaterialNodes/SGraphNodeMaterialBase.cpp
+++ b/Engine/Source/Editor/GraphEditor/Private/MaterialNodes/SGraphNodeMaterialBase.cpp
@@ -264,7 +264,7 @@ void FPreviewElement::UpdateExpressionPreview(UMaterialGraphNode* MaterialNode)
 	);
 }
 
-void FPreviewElement::DrawRenderThread(FRHICommandListImmediate& RHICmdList, const void* InWindowBackBuffer)
+void FPreviewElement::Draw_RenderThread(FRHICommandListImmediate& RHICmdList, const void* InWindowBackBuffer, const FSlateCustomDrawParams& Params)
 {
 	if(ExpressionPreview)
 	{
diff --git a/Engine/Source/Editor/GraphEditor/Public/MaterialNodes/SGraphNodeMaterialBase.h b/Engine/Source/Editor/GraphEditor/Public/MaterialNodes/SGraphNodeMaterialBase.h
--- a/Engine/Source/Editor/GraphEditor/Public/MaterialNodes/SGraphNodeMaterialBase.h
+++ b/Engine/Source/Editor/GraphEditor/Public/MaterialNodes/SGraphNodeMaterialBase.h
@@ -77,7 +77,7 @@ private:
 	/**
 	 * ICustomSlateElement interface 
 	 */
-	virtual void DrawRenderThread(FRHICommandListImmediate& RHICmdList, const void* InWindowBackBuffer) override;
+	virtual void Draw_RenderThread(FRHICommandListImmediate& RHICmdList, const void* InWindowBackBuffer, const FSlateCustomDrawParams& Params) override;
 
 private:
 	/** Render target that the canvas renders to */
diff --git a/Engine/Source/Runtime/Engine/Private/Slate/DebugCanvas.cpp b/Engine/Source/Runtime/Engine/Private/Slate/DebugCanvas.cpp
--- a/Engine/Source/Runtime/Engine/Private/Slate/DebugCanvas.cpp
+++ b/Engine/Source/Runtime/Engine/Private/Slate/DebugCanvas.cpp
@@ -234,7 +234,7 @@ void FDebugCanvasDrawer::InitDebugCanvas(FViewportClient* ViewportClient, UWorld
 	}
 }
 
-void FDebugCanvasDrawer::DrawRenderThread(FRHICommandListImmediate& RHICmdList, const void* InWindowBackBuffer)
+void FDebugCanvasDrawer::Draw_RenderThread(FRHICommandListImmediate& RHICmdList, const void* InWindowBackBuffer, const FSlateCustomDrawParams& Params)
 {
 	check( IsInRenderingThread() );
 	check(RHICmdList.IsOutsideRenderPass());
diff --git a/Engine/Source/Runtime/Engine/Public/Slate/DebugCanvas.h b/Engine/Source/Runtime/Engine/Public/Slate/DebugCanvas.h
--- a/Engine/Source/Runtime/Engine/Public/Slate/DebugCanvas.h
+++ b/Engine/Source/Runtime/Engine/Public/Slate/DebugCanvas.h
@@ -74,7 +74,7 @@ private:
 	/**
 	 * ICustomSlateElement interface 
 	 */
-	virtual void DrawRenderThread(FRHICommandListImmediate& RHICmdList, const void* InWindowBackBuffer) override;
+	virtual void Draw_RenderThread(FRHICommandListImmediate& RHICmdList, const void* InWindowBackBuffer, const FSlateCustomDrawParams& Params) override;
 
 	/**
 	 * Deletes the rendering thread canvas 
diff --git a/Engine/Source/Runtime/Slate/Private/Widgets/Layout/SBackgroundBlur.cpp b/Engine/Source/Runtime/Slate/Private/Widgets/Layout/SBackgroundBlur.cpp
--- a/Engine/Source/Runtime/Slate/Private/Widgets/Layout/SBackgroundBlur.cpp
+++ b/Engine/Source/Runtime/Slate/Private/Widgets/Layout/SBackgroundBlur.cpp
@@ -145,7 +145,7 @@ int32 SBackgroundBlur::OnPaint(const FPaintArgs& Args, const FGeometry& Allotted
 				{
 					OutDrawElements.PushClip(FSlateClippingZone(AllottedGeometry));
 
-					FSlateDrawElement::MakePostProcessPass(OutDrawElements, LayerId, PaintGeometry, FVector4f((float)KernelSize, ComputedStrength, (float)RenderTargetWidth, (float)RenderTargetHeight), DownsampleAmount, FVector4f(CornerRadiusAttribute.Get()));
+					FSlateDrawElement::MakePostProcessBlur(OutDrawElements, LayerId, PaintGeometry, FVector4f((float)KernelSize, ComputedStrength, (float)RenderTargetWidth, (float)RenderTargetHeight), DownsampleAmount, FVector4f(CornerRadiusAttribute.Get()));
 
 					OutDrawElements.PopClip();
 				}
diff --git a/Engine/Source/Runtime/SlateCore/Private/Rendering/DrawElementTypes.cpp b/Engine/Source/Runtime/SlateCore/Private/Rendering/DrawElementTypes.cpp
--- a/Engine/Source/Runtime/SlateCore/Private/Rendering/DrawElementTypes.cpp
+++ b/Engine/Source/Runtime/SlateCore/Private/Rendering/DrawElementTypes.cpp
@@ -689,7 +689,14 @@ void FSlateDrawElement::MakeCustomVerts(FSlateWindowElementList& ElementList, ui
 	Element.RenderTransform = FSlateRenderTransform();
 }
 
+PRAGMA_DISABLE_DEPRECATION_WARNINGS
 void FSlateDrawElement::MakePostProcessPass(FSlateWindowElementList& ElementList, uint32 InLayer, const FPaintGeometry& PaintGeometry, const FVector4f& Params, int32 DownsampleAmount, FVector4f CornerRadius)
+{
+	FSlateDrawElement::MakePostProcessBlur(ElementList, InLayer, PaintGeometry, Params, DownsampleAmount, CornerRadius);
+}
+PRAGMA_ENABLE_DEPRECATION_WARNINGS
+
+void FSlateDrawElement::MakePostProcessBlur(FSlateWindowElementList& ElementList, uint32 InLayer, const FPaintGeometry& PaintGeometry, const FVector4f& Params, int32 DownsampleAmount, FVector4f CornerRadius)
 {
 	PaintGeometry.CommitTransformsIfUsingLegacyConstructor();
 
diff --git a/Engine/Source/Runtime/SlateCore/Private/Rendering/ElementBatcher.cpp b/Engine/Source/Runtime/SlateCore/Private/Rendering/ElementBatcher.cpp
--- a/Engine/Source/Runtime/SlateCore/Private/Rendering/ElementBatcher.cpp
+++ b/Engine/Source/Runtime/SlateCore/Private/Rendering/ElementBatcher.cpp
@@ -284,6 +284,8 @@ FSlateElementBatcher::FSlateElementBatcher( TSharedRef<FSlateRenderingPolicy> In
 	, bRequiresVsync(false)
 	, bCompositeHDRViewports(false)
 	, UsedSlatePostBuffers(ESlatePostRT::None)
+	, ResourceUpdatingPostBuffers(ESlatePostRT::None)
+	, SkipDefaultUpdatePostBuffers(ESlatePostRT::None)
 {
 }
 
@@ -2961,6 +2963,7 @@ void FSlateElementBatcher::AddCustomElement( const FSlateCustomDrawerElement& Dr
 	FSlateRenderBatch& RenderBatch = CreateRenderBatch(Layer, FShaderParams(), nullptr, ESlateDrawPrimitive::None, ESlateShader::Default, ESlateDrawEffect::None, ESlateBatchDrawFlag::None, DrawElement);
 	RenderBatch.CustomDrawer = DrawElement.CustomDrawer.Pin().Get();
 	RenderBatch.bIsMergable = false;
+	RenderBatch.CustomDrawer->PostCustomElementAdded(*this);
 }
 
 void FSlateElementBatcher::AddCustomVerts(const FSlateCustomVertsElement& DrawElement)
@@ -3573,6 +3576,8 @@ void FSlateElementBatcher::ResetBatches()
 	bRequiresVsync = false;
 	bCompositeHDRViewports = false;
 	UsedSlatePostBuffers = ESlatePostRT::None;
+	ResourceUpdatingPostBuffers = ESlatePostRT::None;
+	SkipDefaultUpdatePostBuffers = ESlatePostRT::None;
 	NumPostProcessPasses = 0;
 }
 
diff --git a/Engine/Source/Runtime/SlateCore/Public/Rendering/DrawElementTypes.h b/Engine/Source/Runtime/SlateCore/Public/Rendering/DrawElementTypes.h
--- a/Engine/Source/Runtime/SlateCore/Public/Rendering/DrawElementTypes.h
+++ b/Engine/Source/Runtime/SlateCore/Public/Rendering/DrawElementTypes.h
@@ -250,8 +250,21 @@ public:
 	
 	SLATECORE_API static void MakeCustomVerts(FSlateWindowElementList& ElementList, uint32 InLayer, const FSlateResourceHandle& InRenderResourceHandle, const TArray<FSlateVertex>& InVerts, const TArray<SlateIndex>& InIndexes, ISlateUpdatableInstanceBuffer* InInstanceData, uint32 InInstanceOffset, uint32 InNumInstances, ESlateDrawEffect InDrawEffects = ESlateDrawEffect::None);
 
+	UE_DEPRECATED(5.4, "MakePostProcessPass has been deprecated. If you need to make a blur please call MakePostProcessBlur.")
 	SLATECORE_API static void MakePostProcessPass(FSlateWindowElementList& ElementList, uint32 InLayer, const FPaintGeometry& PaintGeometry, const FVector4f& Params, int32 DownsampleAmount, const FVector4f CornerRadius = FVector4f(0.0f));
 
+	/**
+	 * Creates an element that performs a blur pass
+	 *
+	 * @param ElementList			The list in which to add elements
+	 * @param InLayer				The layer to draw the element on
+	 * @param PaintGeometry			DrawSpace position and dimensions; see FPaintGeometry
+	 * @param Params				Shader params for blur, should be tuple of KernelSize, Strength, Width, & Height
+	 * @param DownSampleAmount		Amount we can downsample for the blur based on kernel size / strength
+	 * @param CornerRadius			Amount pixels will be weighted in any direction
+	 */
+	SLATECORE_API static void MakePostProcessBlur(FSlateWindowElementList& ElementList, uint32 InLayer, const FPaintGeometry& PaintGeometry, const FVector4f& Params, int32 DownsampleAmount, const FVector4f CornerRadius = FVector4f(0.0f));
+
 	FSlateDrawElement();
 
 	FORCEINLINE int32 GetLayer() const { return LayerId; }
diff --git a/Engine/Source/Runtime/SlateCore/Public/Rendering/ElementBatcher.h b/Engine/Source/Runtime/SlateCore/Public/Rendering/ElementBatcher.h
--- a/Engine/Source/Runtime/SlateCore/Public/Rendering/ElementBatcher.h
+++ b/Engine/Source/Runtime/SlateCore/Public/Rendering/ElementBatcher.h
@@ -258,6 +258,14 @@ public:
 
 	ESlatePostRT GetUsedSlatePostBuffers() const { return UsedSlatePostBuffers; }
 
+	ESlatePostRT GetResourceUpdatingPostBuffers() const { return ResourceUpdatingPostBuffers; }
+
+	ESlatePostRT GetSkipDefaultUpdatePostBuffers() const { return SkipDefaultUpdatePostBuffers; }
+
+	void SetResourceUpdatingPostBuffers(ESlatePostRT InResourceUpdatingPostBuffers) { ResourceUpdatingPostBuffers = InResourceUpdatingPostBuffers; }
+
+	void SetSkipDefaultUpdatePostBuffers(ESlatePostRT InSkipDefaultUpdatePostBuffers) { SkipDefaultUpdatePostBuffers = InSkipDefaultUpdatePostBuffers; }
+
 	void SetCompositeHDRViewports(bool bInCompositeHDRViewports) { bCompositeHDRViewports = bInCompositeHDRViewports; }
 
 	/** 
@@ -486,4 +494,10 @@ private:
 
 	// true if we added a resource that is using a slate post buffer 
 	ESlatePostRT UsedSlatePostBuffers;
+
+	// true if a resource is updating a slate post buffer, if true we need to add a fence for the scene draw to complete before clearing unused post buffers
+	ESlatePostRT ResourceUpdatingPostBuffers;
+
+	// true if we should skip the default population of the post buffers with the scene & no UI.
+	ESlatePostRT SkipDefaultUpdatePostBuffers;
 };
diff --git a/Engine/Source/Runtime/SlateCore/Public/Rendering/RenderingCommon.h b/Engine/Source/Runtime/SlateCore/Public/Rendering/RenderingCommon.h
--- a/Engine/Source/Runtime/SlateCore/Public/Rendering/RenderingCommon.h
+++ b/Engine/Source/Runtime/SlateCore/Public/Rendering/RenderingCommon.h
@@ -11,6 +11,7 @@
 #include "Input/NavigationReply.h"
 #include "Input/PopupMethodReply.h"
 #include "Rendering/DrawElementCoreTypes.h"
+#include "Rendering/SlateRendererTypes.h"
 #include "SlateGlobals.h"
 #include <utility>
 
@@ -858,15 +859,69 @@ public:
  */
 class ICustomSlateElement
 {
+public:
+
+	/** Struct describing current draw state for the custom drawer */
+	struct FSlateCustomDrawParams
+	{
+		FMatrix44f ViewProjectionMatrix;
+		FVector2f ViewOffset;
+		FIntRect ViewRect;
+		EDisplayColorGamut HDRDisplayColorGamut;
+		ESlatePostRT UsedSlatePostBuffers;
+		bool bWireFrame;
+		bool bIsHDR;
+
+		FSlateCustomDrawParams()
+			: ViewProjectionMatrix(FMatrix44f())
+			, ViewOffset(0.f, 0.f)
+			, ViewRect(FIntRect())
+			, HDRDisplayColorGamut(EDisplayColorGamut::sRGB_D65)
+			, UsedSlatePostBuffers(ESlatePostRT::None)
+			, bWireFrame(false)
+			, bIsHDR(false)
+		{
+		}
+	};
+
 public:
 	virtual ~ICustomSlateElement() {}
 
+	UE_DEPRECATED(5.4, "Please override Draw_RenderThread instead and modify your function signature to accept 'FSlateCustomParams& Params'")
+	virtual void DrawRenderThread(class FRHICommandListImmediate& RHICmdList, const void* RenderTarget) 
+	{
+	}
+
 	/** 
 	 * Called from the rendering thread when it is time to render the element
 	 *
-	 * @param RenderTarget	handle to the platform specific render target implementation.  Note this is already bound by Slate initially 
+	 * @param RenderTarget				handle to the platform specific render target implementation.  Note this is already bound by Slate initially 
+	 * @param Params					Params about current draw state 
+	 * @param RenderingPolicyInterface	Interface to current rendering policy
+	 */
+	virtual void Draw_RenderThread(class FRHICommandListImmediate& RHICmdList, const void* RenderTarget, const FSlateCustomDrawParams& Params)
+	{
+PRAGMA_DISABLE_DEPRECATION_WARNINGS
+		DrawRenderThread(RHICmdList, RenderTarget);
+PRAGMA_ENABLE_DEPRECATION_WARNINGS
+	}
+
+	/**
+	 * Called from the game thread during element batching
+	 *
+	 * @param ElementBatcher	Elementbatcher that added the custom element
 	 */
-	virtual void DrawRenderThread(class FRHICommandListImmediate& RHICmdList, const void* RenderTarget) = 0;
+	virtual void PostCustomElementAdded(class FSlateElementBatcher& ElementBatcher) const {}
+
+	/**
+	 * If true will cast to an ICustomSlateElementRHI & call Draw_RenderThread with additional RHI params on that instead.
+	 * 
+	 * Note: While a bool to determine cast is not desirable, it is needed due to RHI module reference constraints
+	 */
+	virtual bool UsesAdditionalRHIParams() const 
+	{
+		return false;
+	}
 };
 
 /*
diff --git a/Engine/Source/Runtime/SlateRHIRenderer/Private/FX/SlateFXSubsystem.cpp b/Engine/Source/Runtime/SlateRHIRenderer/Private/FX/SlateFXSubsystem.cpp
--- a/Engine/Source/Runtime/SlateRHIRenderer/Private/FX/SlateFXSubsystem.cpp
+++ b/Engine/Source/Runtime/SlateRHIRenderer/Private/FX/SlateFXSubsystem.cpp
@@ -4,9 +4,49 @@
 
 #include "Engine/Engine.h"
 #include "Engine/GameInstance.h"
+#include "FX/SlateRHIPostBufferProcessor.h"
+#include "SlateRHIRendererSettings.h"
 
 #include UE_INLINE_GENERATED_CPP_BY_NAME(SlateFXSubsystem)
 
+USlateRHIPostBufferProcessor* USlateFXSubsystem::GetPostProcessor(ESlatePostRT InSlatePostBufferBit)
+{
+	USlateRHIPostBufferProcessor* Result = nullptr;
+
+	if (GEngine)
+	{
+		if (USlateFXSubsystem* SlateFXSubsystem = GEngine->GetEngineSubsystem<USlateFXSubsystem>())
+		{
+			Result = SlateFXSubsystem->GetSlatePostProcessor(InSlatePostBufferBit);
+		}
+	}
+
+	return Result;
+}
+
+TSharedPtr<FSlateRHIPostBufferProcessorProxy> USlateFXSubsystem::GetPostProcessorProxy(ESlatePostRT InSlatePostBufferBit)
+{
+	TSharedPtr<FSlateRHIPostBufferProcessorProxy> Result = nullptr;
+
+	if (GEngine)
+	{
+		if (USlateFXSubsystem* SlateFXSubsystem = GEngine->GetEngineSubsystem<USlateFXSubsystem>())
+		{
+			Result = SlateFXSubsystem->GetSlatePostProcessorProxy(InSlatePostBufferBit);
+		}
+	}
+
+	return Result;
+}
+
+void USlateFXSubsystem::BeginDestroy()
+{
+	// Flush rendering commands since this subsystem can be used in render thread
+	FlushRenderingCommands();
+
+	Super::BeginDestroy();
+}
+
 bool USlateFXSubsystem::ShouldCreateSubsystem(UObject* Outer) const
 {
 	TArray<UClass*> ChildClasses;
@@ -34,6 +74,16 @@ void USlateFXSubsystem::Deinitialize()
 	Super::Deinitialize();
 }
 
+TSharedPtr<FSlateRHIPostBufferProcessorProxy> USlateFXSubsystem::GetSlatePostProcessorProxy(ESlatePostRT InPostBufferBit)
+{
+	if (TSharedPtr<FSlateRHIPostBufferProcessorProxy>* ProcessorProxy = SlatePostBufferProcessorProxies.Find(InPostBufferBit))
+	{
+		return *ProcessorProxy;
+	}
+
+	return nullptr;
+}
+
 USlateRHIPostBufferProcessor* USlateFXSubsystem::GetSlatePostProcessor(ESlatePostRT InPostBufferBit)
 {
 	if (TObjectPtr<USlateRHIPostBufferProcessor>* Processor = SlatePostBufferProcessors.Find(InPostBufferBit))
@@ -47,6 +97,7 @@ USlateRHIPostBufferProcessor* USlateFXSubsystem::GetSlatePostProcessor(ESlatePos
 void USlateFXSubsystem::OnPreWorldInitialization(UWorld* World, const UWorld::InitializationValues IVS)
 {
 	SlatePostBufferProcessors.Empty();
+	SlatePostBufferProcessorProxies.Empty();
 
 	if (const USlateRHIRendererSettings* SlateRendererSettings = USlateRHIRendererSettings::Get())
 	{
@@ -55,7 +106,11 @@ void USlateFXSubsystem::OnPreWorldInitialization(UWorld* World, const UWorld::In
 			const FSlatePostSettings& PostSetting = SlateRendererSettings->GetSlatePostSetting(SlatePostBufferBit);
 			if (PostSetting.bEnabled && PostSetting.PostProcessorClass)
 			{
-				SlatePostBufferProcessors.Add(SlatePostBufferBit, NewObject<USlateRHIPostBufferProcessor>(this, PostSetting.PostProcessorClass));
+				if (TObjectPtr<USlateRHIPostBufferProcessor> BufferProcessor = NewObject<USlateRHIPostBufferProcessor>(this, PostSetting.PostProcessorClass))
+				{
+					SlatePostBufferProcessors.Add(SlatePostBufferBit, BufferProcessor);
+					SlatePostBufferProcessorProxies.Add(SlatePostBufferBit, BufferProcessor->GetRenderThreadProxy());
+				}
 			}
 		}
 	}
@@ -64,4 +119,5 @@ void USlateFXSubsystem::OnPreWorldInitialization(UWorld* World, const UWorld::In
 void USlateFXSubsystem::OnPostWorldCleanup(UWorld* World, bool SessionEnded, bool bCleanupResources)
 {
 	SlatePostBufferProcessors.Empty();
+	SlatePostBufferProcessorProxies.Empty();
 }
diff --git a/Engine/Source/Runtime/SlateRHIRenderer/Private/FX/SlatePostBufferBlur.cpp b/Engine/Source/Runtime/SlateRHIRenderer/Private/FX/SlatePostBufferBlur.cpp
--- a/Engine/Source/Runtime/SlateRHIRenderer/Private/FX/SlatePostBufferBlur.cpp
+++ b/Engine/Source/Runtime/SlateRHIRenderer/Private/FX/SlatePostBufferBlur.cpp
@@ -7,25 +7,102 @@
 
 #include UE_INLINE_GENERATED_CPP_BY_NAME(SlatePostBufferBlur)
 
-void USlatePostBufferBlur::PostProcess(FRenderResource* InViewInfo, FRenderResource* InViewportTexture, FVector2D InElementWindowSize, FSlateRHIRenderingPolicyInterface InRenderingPolicy, UTextureRenderTarget2D* InSlatePostBuffer)
+/////////////////////////////////////////////////////
+// FSlatePostBufferBlurProxy
+
+void FSlatePostBufferBlurProxy::PostProcess_Renderthread(FRHICommandListImmediate& RHICmdList, FRHITexture* Src, FRHITexture* Dst, FIntRect SrcRect, FIntRect DstRect, FSlateRHIRenderingPolicyInterface InRenderingPolicy)
 {
 	if (InRenderingPolicy.IsValid())
 	{
-		// Explicit param copy to avoid renderthread from reading value during gamethread write
-		float GaussianBlurStrengthCopy = GaussianBlurStrength;
+		// Use rendering policy to perform blur post process with desired Src / Dst & respective extents
+		InRenderingPolicy.BlurRectExternal(RHICmdList, Src, Dst, SrcRect, DstRect, GaussianBlurStrength_RenderThread);
+	}
+}
+
+void FSlatePostBufferBlurProxy::OnUpdateValuesRenderThread()
+{
+	// Don't issue multiple updates in a single draw
+	if (!ParamUpdateFence.IsFenceComplete())
+	{
+		return;
+	}
 
-		ENQUEUE_RENDER_COMMAND(FUpdateSlatePostBuffersWithFX_Blur)([InViewInfo, InViewportTexture, InElementWindowSize, InRenderingPolicy, InSlatePostBuffer, GaussianBlurStrengthCopy](FRHICommandListImmediate& RHICmdList)
+	// Only issue an update when parent exists & values are different
+	if (USlatePostBufferBlur* ParentBlurObject = Cast<USlatePostBufferBlur>(ParentObject))
+	{
+		if (ParentBlurObject->GaussianBlurStrength != GaussianBlurStrength_RenderThread)
 		{
-			FTexture2DRHIRef BackBuffer = USlateRHIPostBufferProcessor::GetBackbuffer_RenderThread(InViewInfo, InViewportTexture, InElementWindowSize, RHICmdList);
+			// Explicit param copy to avoid renderthread from reading value during gamethread write
+			float GaussianBlurStrengthCopy = ParentBlurObject->GaussianBlurStrength;
 
-			if (BackBuffer)
+			// Execute param copy in a render command to safely update value on renderthread without race conditions
+			TWeakPtr<FSlatePostBufferBlurProxy> TempWeakThis = SharedThis(this);
+			ENQUEUE_RENDER_COMMAND(FUpdateValuesRenderThreadFX_Blur)([TempWeakThis, GaussianBlurStrengthCopy](FRHICommandListImmediate& RHICmdList)
 			{
-				FTexture2DRHIRef Src = USlateRHIPostBufferProcessor::GetSrcTexture_RenderThread(BackBuffer, InViewportTexture);
-				FTextureReferenceRHIRef& Dst = USlateRHIPostBufferProcessor::GetDstTexture_RenderThread(InSlatePostBuffer);
-				FIntPoint DstExtent = USlateRHIPostBufferProcessor::GetDstExtent_RenderThread(BackBuffer, InViewportTexture);
+				if (TSharedPtr<FSlatePostBufferBlurProxy> SharedThisPin = TempWeakThis.Pin())
+				{
+					SharedThisPin->GaussianBlurStrength_RenderThread = GaussianBlurStrengthCopy;
+				}
+			});
+
+			// Issue fence to prevent multiple updates in a single draw
+			ParamUpdateFence.BeginFence();
+		}
+	}
+}
+
+/////////////////////////////////////////////////////
+// USlatePostBufferBlur
 
-				InRenderingPolicy.BlurRectExternal(RHICmdList, Src, Dst, DstExtent, GaussianBlurStrengthCopy);
-			}
-		});
+USlatePostBufferBlur::USlatePostBufferBlur()
+{
+	RenderThreadProxy = nullptr;
+}
+
+USlatePostBufferBlur::~USlatePostBufferBlur()
+{
+	RenderThreadProxy = nullptr;
+}
+
+void USlatePostBufferBlur::PostProcess(FRenderResource* InViewInfo, FRenderResource* InViewportTexture, FVector2D InElementWindowSize, FSlateRHIRenderingPolicyInterface InRenderingPolicy, UTextureRenderTarget2D* InSlatePostBuffer)
+{
+	if (!InRenderingPolicy.IsValid())
+	{
+		return;
+	};
+
+	// Explicit param copy to avoid renderthread from reading value during gamethread write
+	float GaussianBlurStrengthCopy = GaussianBlurStrength;
+
+	// Enqueue default post process command, will trigger on scene before any other slate element draws
+	ENQUEUE_RENDER_COMMAND(FUpdateSlatePostBuffersWithFX_Blur)([InViewInfo, InViewportTexture, InElementWindowSize, InRenderingPolicy, InSlatePostBuffer, GaussianBlurStrengthCopy](FRHICommandListImmediate& RHICmdList)
+	{
+		// Get Backbuffer, which can vary between PIE or standalone
+		FTexture2DRHIRef BackBuffer = USlateRHIPostBufferProcessor::GetBackbuffer_RenderThread(InViewInfo, InViewportTexture, InElementWindowSize, RHICmdList);
+
+		if (BackBuffer)
+		{
+			// Get Src / Dst textures & their rects, again may vary between PIE or standalone
+			// Here we can simply use Src Rect since the Src texture in PIE is the 'BufferedRT' scene backbuffer without the editor
+			FTexture2DRHIRef Src = USlateRHIPostBufferProcessor::GetSrcTexture_RenderThread(BackBuffer, InViewportTexture);
+			FTextureReferenceRHIRef& Dst = USlateRHIPostBufferProcessor::GetDstTexture_RenderThread(InSlatePostBuffer);
+			FIntPoint DstExtent = USlateRHIPostBufferProcessor::GetDstExtent_RenderThread(BackBuffer, InViewportTexture);
+			FIntRect SrcRect = FIntRect(0, 0, Src->GetSizeX(), Src->GetSizeY());
+			FIntRect DstRect = FIntRect(0, 0, DstExtent.X, DstExtent.Y);
+
+			// Use rendering policy to perform blur post process with desired Src / Dst & respective extents
+			InRenderingPolicy.BlurRectExternal(RHICmdList, Src, Dst, SrcRect, DstRect, GaussianBlurStrengthCopy);
+		}
+	});
+}
+
+TSharedPtr<FSlateRHIPostBufferProcessorProxy> USlatePostBufferBlur::GetRenderThreadProxy()
+{
+	if (!RenderThreadProxy && IsInGameThread())
+	{
+		// Create a RT proxy specific for doing blurs
+		RenderThreadProxy = MakeShared<FSlatePostBufferBlurProxy>();
+		RenderThreadProxy->SetOwningProcessorObject(this);
 	}
-}
\ No newline at end of file
+	return RenderThreadProxy;
+}
diff --git a/Engine/Source/Runtime/SlateRHIRenderer/Private/SlatePostProcessor.cpp b/Engine/Source/Runtime/SlateRHIRenderer/Private/SlatePostProcessor.cpp
--- a/Engine/Source/Runtime/SlateRHIRenderer/Private/SlatePostProcessor.cpp
+++ b/Engine/Source/Runtime/SlateRHIRenderer/Private/SlatePostProcessor.cpp
@@ -202,9 +202,14 @@ static void BlitUIToHDRScene(FRHICommandListImmediate& RHICmdList, IRendererModu
 
 	const FVector2f InvSrcTextureSize(1.f / SrcTextureWidth, 1.f / SrcTextureHeight);
 
+	// If using a custom output target, sample the source rect for UV's, else do not as the dest rect may be a subsection
+	const FSlateRect& UVRect = RectParams.PostProcessDest == EPostProcessDestination::DestTexture
+		? SourceRect
+		: DestRect;
+
 	// no guard band is actually needed because GaussianBlurMain already takes ensure that sampling happens inside the rectangle
-	const FVector2f UVStart = FVector2f(DestRect.Left - 0.0f, DestRect.Top - 0.0f) * InvSrcTextureSize;
-	const FVector2f UVEnd = FVector2f(DestRect.Right + 0.0f, DestRect.Bottom + 0.0f) * InvSrcTextureSize;
+	const FVector2f UVStart = FVector2f(UVRect.Left - 0.0f, UVRect.Top - 0.0f) * InvSrcTextureSize;
+	const FVector2f UVEnd = FVector2f(UVRect.Right + 0.0f, UVRect.Bottom + 0.0f) * InvSrcTextureSize;
 	const FVector2f SizeUV = UVEnd - UVStart;
 
 	RHICmdList.SetViewport(0, 0, 0, SrcTextureWidth, SrcTextureHeight, 0.0f);
@@ -369,8 +374,13 @@ void FSlatePostProcessor::BlurRect(FRHICommandListImmediate& RHICmdList, IRender
 				{
 					const FVector2f InvSrcTextureSize(1.f / SrcTextureWidth, 1.f / SrcTextureHeight);
 
-					const FVector2f UVStart = FVector2f(DestRect.Left, DestRect.Top) * InvSrcTextureSize;
-					const FVector2f UVEnd = FVector2f(DestRect.Right, DestRect.Bottom) * InvSrcTextureSize;
+					// If using a custom output target, sample the source rect for UV's, else do not as the dest rect may be a subsection
+					const FSlateRect& UVRect = RectParams.PostProcessDest == EPostProcessDestination::DestTexture
+						? SourceRect
+						: DestRect;
+
+					const FVector2f UVStart = FVector2f(UVRect.Left, UVRect.Top) * InvSrcTextureSize;
+					const FVector2f UVEnd = FVector2f(UVRect.Right, UVRect.Bottom) * InvSrcTextureSize;
 					const FVector2f SizeUV = UVEnd - UVStart;
 
 					PixelShader->SetUVBounds(BatchedParameters, FVector4f(UVStart, UVEnd));
@@ -571,8 +581,13 @@ void FSlatePostProcessor::DownsampleRect(FRHICommandListImmediate& RHICmdList, I
 
 		const FVector2f InvSrcTextureSize(1.f/SrcTextureWidth, 1.f/SrcTextureHeight);
 
-		const FVector2f UVStart = FVector2f(DestRect.Left, DestRect.Top) * InvSrcTextureSize;
-		const FVector2f UVEnd = FVector2f(DestRect.Right, DestRect.Bottom) * InvSrcTextureSize;
+		// If using a custom output target, sample the source rect for UV's, else do not as the dest rect may be a subsection
+		const FSlateRect& UVRect = Params.PostProcessDest == EPostProcessDestination::DestTexture
+			? SourceRect
+			: DestRect;
+
+		const FVector2f UVStart = FVector2f(UVRect.Left, UVRect.Top) * InvSrcTextureSize;
+		const FVector2f UVEnd = FVector2f(UVRect.Right, UVRect.Bottom) * InvSrcTextureSize;
 		const FVector2f SizeUV = UVEnd - UVStart;
 		
 		RHICmdList.SetViewport(0.f, 0.f, 0.f, (float)DestTextureWidth, (float)DestTextureHeight, 0.0f);
@@ -636,8 +651,8 @@ void FSlatePostProcessor::UpsampleRect(FRHICommandListImmediate& RHICmdList, IRe
 	FTexture2DRHIRef DestTexture = Params.PostProcessDest == EPostProcessDestination::DestTexture && Params.DestTexture 
 		? Params.DestTexture 
 		: Params.SourceTexture;
-	const int32 DestTextureWidth = Params.SourceTextureSize.X;
-	const int32 DestTextureHeight = Params.SourceTextureSize.Y;
+	const int32 DestTextureWidth = DestTexture->GetSizeX();
+	const int32 DestTextureHeight = DestTexture->GetSizeY();
 
 	const int32 DownsampledWidth = DownsampleSize.X;
 	const int32 DownsampledHeight = DownsampleSize.Y;
@@ -737,7 +752,7 @@ void FSlatePostProcessor::UpsampleRect(FRHICommandListImmediate& RHICmdList, IRe
 			Size.X, Size.Y,
 			0, 0,
 			SizeUV.X, SizeUV.Y,
-			Params.SourceTextureSize,
+			DestTexture->GetSizeXY(),
 			FIntPoint(1, 1),
 			VertexShader,
 			EDRF_Default);
diff --git a/Engine/Source/Runtime/SlateRHIRenderer/Private/SlateRHIRenderer.cpp b/Engine/Source/Runtime/SlateRHIRenderer/Private/SlateRHIRenderer.cpp
--- a/Engine/Source/Runtime/SlateRHIRenderer/Private/SlateRHIRenderer.cpp
+++ b/Engine/Source/Runtime/SlateRHIRenderer/Private/SlateRHIRenderer.cpp
@@ -138,6 +138,7 @@ struct FSlateDrawWindowCommandParams
 	FSlateWindowElementList* WindowElementList;
 	SWindow* Window;
 	FIntRect ViewRect;
+	ESlatePostRT UsedSlatePostBuffers;
 #if WANTS_DRAW_MESH_EVENTS
 	FString WindowTitle;
 #endif
@@ -276,6 +277,11 @@ int32 FSlateRHIRenderer::GetDrawToVRRenderTarget()
 	return CVarDrawToVRRenderTarget->GetInt();
 }
 
+int32 FSlateRHIRenderer::GetProcessSlatePostBuffers()
+{
+	return CVarCopyBackbufferToSlatePostRenderTargets->GetInt();
+}
+
 bool FSlateRHIRenderer::Initialize()
 {
 	LoadUsedTextures();
@@ -1001,6 +1007,7 @@ void RenderSlateBatch(FTexture2DRHIRef SlateRenderTarget, bool bClear, bool bIsH
 				RenderParams.bIsHDR = bIsHDR;
 				RenderParams.HDRDisplayColorGamut = ViewportInfo.HDRDisplayColorGamut;
 				RenderParams.ViewRect = DrawCommandParams.ViewRect;
+				RenderParams.UsedSlatePostBuffers = DrawCommandParams.UsedSlatePostBuffers;
 				if (ViewportInfo.bSceneHDREnabled && !bIsHDR)
 				{
 					RenderParams.UITarget = ViewportInfo.UITargetRT;
@@ -1730,8 +1737,10 @@ void FSlateRHIRenderer::DrawWindows_Private(FSlateDrawBuffer& WindowDrawBuffer)
 
 				bLockToVsync |= bForceVsyncFromCVar;
 
-				// Cache if the element batcher used a post buffer, will get reset on batch reset
+				// Cache if the element batcher post buffer usage, will get reset on batch reset
 				ESlatePostRT UsedSlatePostBufferBits = ElementBatcher->GetUsedSlatePostBuffers();
+				ESlatePostRT ResourceUpdatingPostBufferBits = ElementBatcher->GetResourceUpdatingPostBuffers();
+				ESlatePostRT SkipDefaultUpdatePostBufferBits = ElementBatcher->GetSkipDefaultUpdatePostBuffers();
 
 				// All elements for this window have been batched and rendering data updated
 				ElementBatcher->ResetBatches();
@@ -1766,6 +1775,7 @@ void FSlateRHIRenderer::DrawWindows_Private(FSlateDrawBuffer& WindowDrawBuffer)
 
 						bool bIsViewportPresentForPIE = GIsEditor ? Window->GetViewport().IsValid() : true;
 						bool bPostBufferBitUsed = (UsedSlatePostBufferBits & SlatePostBufferBit) != ESlatePostRT::None;
+						bool bSkipDefaultPostBufferUpdate = (SkipDefaultUpdatePostBufferBits & SlatePostBufferBit) != ESlatePostRT::None;
 
 						if (bIsViewportPresentForPIE && bPostBufferBitUsed)
 						{
@@ -1793,61 +1803,49 @@ void FSlateRHIRenderer::DrawWindows_Private(FSlateDrawBuffer& WindowDrawBuffer)
 
 							const FVector2D ElementWindowSize = ElementList.GetWindowSize();
 
-							auto GetPostProcessor = [&](ESlatePostRT InSlatePostBufferBit)
+							if (!bSkipDefaultPostBufferUpdate)
 							{
-								USlateRHIPostBufferProcessor* Result = nullptr;
-
-								if (GEngine)
+								if (USlateRHIPostBufferProcessor* PostProcessor = USlateFXSubsystem::GetPostProcessor(SlatePostBufferBit))
 								{
-									if (USlateFXSubsystem* SlateFXSubsystem = GEngine->GetEngineSubsystem<USlateFXSubsystem>())
-									{
-										Result = SlateFXSubsystem->GetSlatePostProcessor(InSlatePostBufferBit);
-									}
+									// Allow the post processor to enque render commands, we delgate this task so it can be done in a thread safe manner.
+									PostProcessor->PostProcess(ViewInfo, ViewportTexture, ElementWindowSize, FSlateRHIRenderingPolicyInterface(RenderingPolicy.Get()), SlatePostBuffer);
 								}
-
-								return Result;
-							};
-
-							if (USlateRHIPostBufferProcessor* PostProcessor = GetPostProcessor(SlatePostBufferBit))
-							{
-								// Allow the post processor to enque render commands, we delgate this task so it can be done in a thread safe manner.
-								PostProcessor->PostProcess(ViewInfo, ViewportTexture, ElementWindowSize, FSlateRHIRenderingPolicyInterface(RenderingPolicy.Get()), SlatePostBuffer);
-							}
-							else
-							{
-								ENQUEUE_RENDER_COMMAND(FUpdateSlatePostBuffers)([ViewInfo, ElementWindowSize, ViewportTexture, SlatePostBuffer](FRHICommandListImmediate& RHICmdList)
+								else
 								{
-									bool bRenderedStereo = false;
-									if (CVarDrawToVRRenderTarget->GetInt() == 0 && GEngine && IsValidRef(ViewInfo->GetRenderTargetTexture()) && GEngine->StereoRenderingDevice.IsValid())
-									{
-										GEngine->StereoRenderingDevice->RenderTexture_RenderThread(RHICmdList, RHIGetViewportBackBuffer(ViewInfo->ViewportRHI), ViewInfo->GetRenderTargetTexture(), ElementWindowSize);
-										bRenderedStereo = true;
-									}
-
-									FTexture2DRHIRef ViewportRT = bRenderedStereo ? nullptr : ViewInfo->GetRenderTargetTexture();
-									FTexture2DRHIRef BackBuffer = (ViewportRT) ? ViewportRT : RHIGetViewportBackBuffer(ViewInfo->ViewportRHI);
-
-									if (BackBuffer)
+									ENQUEUE_RENDER_COMMAND(FUpdateSlatePostBuffers)([ViewInfo, ElementWindowSize, ViewportTexture, SlatePostBuffer](FRHICommandListImmediate& RHICmdList)
 									{
-										FRHICopyTextureInfo CopyInfo;
-
-										// Copy just the viewport RT if in PIE, else do entire backbuffer
-										if (GIsEditor)
+										bool bRenderedStereo = false;
+										if (CVarDrawToVRRenderTarget->GetInt() == 0 && GEngine && IsValidRef(ViewInfo->GetRenderTargetTexture()) && GEngine->StereoRenderingDevice.IsValid())
 										{
-											CopyInfo.Size = FIntVector(ViewportTexture->GetWidth(), ViewportTexture->GetHeight(), 1);
-											TransitionAndCopyTexture(RHICmdList, ViewportTexture->GetRHIRef(), SlatePostBuffer->TextureReference.TextureReferenceRHI, CopyInfo);
+											GEngine->StereoRenderingDevice->RenderTexture_RenderThread(RHICmdList, RHIGetViewportBackBuffer(ViewInfo->ViewportRHI), ViewInfo->GetRenderTargetTexture(), ElementWindowSize);
+											bRenderedStereo = true;
 										}
-										else
+
+										FTexture2DRHIRef ViewportRT = bRenderedStereo ? nullptr : ViewInfo->GetRenderTargetTexture();
+										FTexture2DRHIRef BackBuffer = (ViewportRT) ? ViewportRT : RHIGetViewportBackBuffer(ViewInfo->ViewportRHI);
+
+										if (BackBuffer)
 										{
-											FIntPoint BackbufferExtent = BackBuffer->GetDesc().Extent;
-											CopyInfo.Size = FIntVector(BackbufferExtent.X, BackbufferExtent.Y, 1);
-											TransitionAndCopyTexture(RHICmdList, BackBuffer, SlatePostBuffer->TextureReference.TextureReferenceRHI, CopyInfo);
+											FRHICopyTextureInfo CopyInfo;
+
+											// Copy just the viewport RT if in PIE, else do entire backbuffer
+											if (GIsEditor)
+											{
+												CopyInfo.Size = FIntVector(ViewportTexture->GetWidth(), ViewportTexture->GetHeight(), 1);
+												TransitionAndCopyTexture(RHICmdList, ViewportTexture->GetRHIRef(), SlatePostBuffer->TextureReference.TextureReferenceRHI, CopyInfo);
+											}
+											else
+											{
+												FIntPoint BackbufferExtent = BackBuffer->GetDesc().Extent;
+												CopyInfo.Size = FIntVector(BackbufferExtent.X, BackbufferExtent.Y, 1);
+												TransitionAndCopyTexture(RHICmdList, BackBuffer, SlatePostBuffer->TextureReference.TextureReferenceRHI, CopyInfo);
+											}
 										}
-									}
-								});
-							}
+									});
+								}
 
-							SlatePostRTFences[SlatePostBufferBitIndex].BeginFence();
+								SlatePostRTFences[SlatePostBufferBitIndex].BeginFence();
+							}
 
 							bShrinkPostBufferRequested &= ~SlatePostBufferBit;
 						}
@@ -1901,6 +1899,7 @@ void FSlateRHIRenderer::DrawWindows_Private(FSlateDrawBuffer& WindowDrawBuffer)
 					Params.WindowElementList = &ElementList;
 					Params.Window = Window;
 					Params.ViewRect = GetViewRect();
+					Params.UsedSlatePostBuffers = UsedSlatePostBufferBits;
 #if WANTS_DRAW_MESH_EVENTS
 					Params.WindowTitle = Window->GetTitle().ToString();
 #endif
@@ -1924,6 +1923,23 @@ void FSlateRHIRenderer::DrawWindows_Private(FSlateDrawBuffer& WindowDrawBuffer)
 								Params.Renderer->DrawWindow_RenderThread(RHICmdList, *ViewInfo, *Params.WindowElementList, Params);
 							}
 						);
+
+						// After we draw, if a resource is going to update a post buffer we need to add another render fence for that buffer
+						if (CVarCopyBackbufferToSlatePostRenderTargets.GetValueOnAnyThread())
+						{
+							uint8 SlatePostBufferBitIndex = 0;
+							for (ESlatePostRT SlatePostBufferBit : TEnumRange<ESlatePostRT>())
+							{
+								bool bResourceUpdatedPostBuffer = (ResourceUpdatingPostBufferBits & SlatePostBufferBit) != ESlatePostRT::None;
+								if (bResourceUpdatedPostBuffer)
+								{
+									SlatePostRTFences[SlatePostBufferBitIndex].BeginFence();
+									bShrinkPostBufferRequested &= ~SlatePostBufferBit;
+								}
+
+								SlatePostBufferBitIndex++;
+							}
+						}
 					}
 
 					SlateWindowRendered.Broadcast(*Window, &ViewInfo->ViewportRHI);
@@ -2425,6 +2441,11 @@ void FSlateRHIRenderer::SetWindowRenderTarget(const SWindow& Window, IViewportRe
 	}
 }
 
+FSlateRHIRenderingPolicyInterface FSlateRHIRenderer::GetRenderingPolicyInterface()
+{
+	return FSlateRHIRenderingPolicyInterface(RenderingPolicy.Get());
+}
+
 void FSlateRHIRenderer::AddWidgetRendererUpdate(const struct FRenderThreadUpdateContext& Context, bool bDeferredRenderTargetUpdate)
 {
 	if (bDeferredRenderTargetUpdate)
diff --git a/Engine/Source/Runtime/SlateRHIRenderer/Private/SlateRHIRenderer.h b/Engine/Source/Runtime/SlateRHIRenderer/Private/SlateRHIRenderer.h
--- a/Engine/Source/Runtime/SlateRHIRenderer/Private/SlateRHIRenderer.h
+++ b/Engine/Source/Runtime/SlateRHIRenderer/Private/SlateRHIRenderer.h
@@ -8,7 +8,6 @@
 #include "RHI.h"
 #include "RenderCommandFence.h"
 #include "RenderResource.h"
-#include "SlateRHIRendererSettings.h"
 #include "SlateRHIResourceManager.h"
 #include "UnrealClient.h"
 #include "Rendering/SlateRenderer.h"
@@ -18,8 +17,12 @@
 
 class FSlateElementBatcher;
 class FSlateRHIRenderingPolicy;
+class FSlateRHIRenderingPolicyInterface;
+class USlateRHIPostBufferProcessor;
+class USlateRHIRendererSettings;
 class ISlateStyle;
 class SWindow;
+struct FSlatePostSettings;
 struct Rect;
 
 template<typename TCmd, typename NameType> struct FRHICommand;
@@ -178,6 +181,7 @@ public:
 	 */
 	static FMatrix CreateProjectionMatrix( uint32 Width, uint32 Height );
 	static int32 GetDrawToVRRenderTarget();
+	static int32 GetProcessSlatePostBuffers();
 
 	/** FSlateRenderer interface */
 	virtual bool Initialize() override;
@@ -244,6 +248,7 @@ public:
 
 	virtual void SetWindowRenderTarget(const SWindow& Window, class IViewportRenderTargetProvider* Provider) override;
 
+	FSlateRHIRenderingPolicyInterface GetRenderingPolicyInterface();
 
 private:
 	/** Loads all known textures from Slate styles */
diff --git a/Engine/Source/Runtime/SlateRHIRenderer/Private/SlateRHIRendererSettings.cpp b/Engine/Source/Runtime/SlateRHIRenderer/Private/SlateRHIRendererSettings.cpp
--- a/Engine/Source/Runtime/SlateRHIRenderer/Private/SlateRHIRendererSettings.cpp
+++ b/Engine/Source/Runtime/SlateRHIRenderer/Private/SlateRHIRendererSettings.cpp
@@ -1,6 +1,7 @@
 // Copyright Epic Games, Inc. All Rights Reserved.
 
 #include "SlateRHIRendererSettings.h"
+#include "FX/SlateRHIPostBufferProcessor.h"
 #include "HAL/IConsoleManager.h"
 
 #include UE_INLINE_GENERATED_CPP_BY_NAME(SlateRHIRendererSettings)
@@ -52,6 +53,14 @@ USlateRHIRendererSettings::~USlateRHIRendererSettings()
 	}
 }
 
+void USlateRHIRendererSettings::BeginDestroy()
+{
+	// Flush rendering commands since these settings can be used in render thread
+	FlushRenderingCommands();
+
+	Super::BeginDestroy();
+}
+
 FSlatePostSettings& USlateRHIRendererSettings::GetMutableSlatePostSetting(ESlatePostRT InPostBufferBit)
 {
 	return SlatePostSettings[InPostBufferBit];
diff --git a/Engine/Source/Runtime/SlateRHIRenderer/Private/SlateRHIRenderingPolicy.cpp b/Engine/Source/Runtime/SlateRHIRenderer/Private/SlateRHIRenderingPolicy.cpp
--- a/Engine/Source/Runtime/SlateRHIRenderer/Private/SlateRHIRenderingPolicy.cpp
+++ b/Engine/Source/Runtime/SlateRHIRenderer/Private/SlateRHIRenderingPolicy.cpp
@@ -10,6 +10,7 @@
 #include "RHIStaticStates.h"
 #include "RHIUtilities.h"
 #include "GlobalRenderResources.h"
+#include "Interfaces/SlateRHIRenderingPolicyInterface.h"
 #include "SceneView.h"
 #include "SceneUtils.h"
 #include "Engine/Engine.h"
@@ -1244,7 +1245,7 @@ void FSlateRHIRenderingPolicy::DrawElements(
 
 				FPostProcessRectParams RectParams;
 				RectParams.SourceTexture = PostProcessTexture;
-				RectParams.SourceRect = FSlateRect(0.f, 0.f, (float)PostProcessTexture->GetSizeX(), (float)PostProcessTexture->GetSizeY());
+				RectParams.SourceRect = FSlateRect((float)Params.ViewRect.Min.X, (float)Params.ViewRect.Min.Y, (float)Params.ViewRect.Max.X, (float)Params.ViewRect.Max.Y);
 				RectParams.DestRect = FSlateRect(QuadPositionData.X, QuadPositionData.Y, QuadPositionData.Z, QuadPositionData.W);
 				RectParams.SourceTextureSize = PostProcessTexture->GetSizeXY();
 				RectParams.CornerRadius = ShaderParams.PixelParams3;
@@ -1300,8 +1301,25 @@ void FSlateRHIRenderingPolicy::DrawElements(
 				RHICmdList.SetScissorRect(false, 0, 0, 0, 0);
 				LastClippingState = nullptr;
 
+				ICustomSlateElement::FSlateCustomDrawParams CustomDrawParams = ICustomSlateElement::FSlateCustomDrawParams();
+				CustomDrawParams.ViewProjectionMatrix = Params.ViewProjectionMatrix;
+				CustomDrawParams.ViewOffset = Params.ViewOffset;
+				CustomDrawParams.ViewRect = Params.ViewRect;
+				CustomDrawParams.HDRDisplayColorGamut = Params.HDRDisplayColorGamut;
+				CustomDrawParams.UsedSlatePostBuffers = Params.UsedSlatePostBuffers;
+				CustomDrawParams.bWireFrame = Params.bWireFrame;
+				CustomDrawParams.bIsHDR = Params.bIsHDR;
+
 				// This element is custom and has no Slate geometry.  Tell it to render itself now
-				CustomDrawer->DrawRenderThread(RHICmdList, &BackBuffer.GetRenderTargetTexture());
+				if (CustomDrawer->UsesAdditionalRHIParams())
+				{
+					ICustomSlateElementRHI* CustomDrawerRHI = static_cast<ICustomSlateElementRHI*>(CustomDrawer);
+					CustomDrawerRHI->Draw_RHIRenderThread(RHICmdList, BackBuffer.GetRenderTargetTexture(), CustomDrawParams, FSlateRHIRenderingPolicyInterface(this));
+				}
+				else
+				{
+					CustomDrawer->Draw_RenderThread(RHICmdList, &BackBuffer.GetRenderTargetTexture(), CustomDrawParams);
+				}
 
 				//We reset the maskingID here because otherwise the RT might not get re-set in the lines above see: if (bClearStencil || bForceStateChange)
 				MaskingID = 0;
@@ -1654,11 +1672,11 @@ void FSlateRHIRenderingPolicy::FlushGeneratedResources()
 	PostProcessor->ReleaseRenderTargets();
 }
 
-void FSlateRHIRenderingPolicy::BlurRectExternal(FRHICommandListImmediate& RHICmdList, FTexture2DRHIRef BlurSrc, FTextureReferenceRHIRef& BlurDst, FIntPoint DstExtent, float BlurStrength) const
+void FSlateRHIRenderingPolicy::BlurRectExternal(FRHICommandListImmediate& RHICmdList, FRHITexture* BlurSrc, FRHITexture* BlurDst, FIntRect SrcRect, FIntRect DstRect, float BlurStrength) const
 {
 	SLATE_DRAW_EVENT(RHICmdList, PostProcess);
 
-	FIntPoint BlurDstExtent = DstExtent;
+	FIntPoint BlurDstExtent = FIntPoint(DstRect.Width(), DstRect.Height());
 
 	// If the radius isn't set, auto-compute it based on the strength
 	int32 OutKernelSize = FMath::RoundToInt(BlurStrength * 3.f);
@@ -1698,7 +1716,7 @@ void FSlateRHIRenderingPolicy::BlurRectExternal(FRHICommandListImmediate& RHICmd
 
 	FPostProcessRectParams RectParams;
 	RectParams.SourceTexture = BlurSrc;
-	RectParams.SourceRect = FSlateRect(0.f, 0.f, (float)BlurSrc->GetSizeX(), (float)BlurSrc->GetSizeY());
+	RectParams.SourceRect = FSlateRect((float)SrcRect.Min.X, (float)SrcRect.Min.Y, (float)SrcRect.Max.X, (float)SrcRect.Max.Y);
 	RectParams.DestRect = FSlateRect(TopLeft.X, TopLeft.Y, BotRight.X, BotRight.Y);
 	RectParams.SourceTextureSize = BlurSrc->GetSizeXY();
 	RectParams.CornerRadius = FVector4f(0, 0, 0, 0);
diff --git a/Engine/Source/Runtime/SlateRHIRenderer/Private/SlateRHIRenderingPolicy.h b/Engine/Source/Runtime/SlateRHIRenderer/Private/SlateRHIRenderingPolicy.h
--- a/Engine/Source/Runtime/SlateRHIRenderer/Private/SlateRHIRenderingPolicy.h
+++ b/Engine/Source/Runtime/SlateRHIRenderer/Private/SlateRHIRenderingPolicy.h
@@ -7,6 +7,7 @@
 #include "GameTime.h"
 #include "RendererInterface.h"
 #include "Rendering/RenderingCommon.h"
+#include "Rendering/SlateRendererTypes.h"
 #include "Rendering/ShaderResourceManager.h"
 #include "Rendering/DrawElements.h"
 #include "Rendering/RenderingPolicy.h"
@@ -36,6 +37,7 @@ struct FSlateRenderingParams
 	FGameTime Time;
 	TRefCountPtr<IPooledRenderTarget> UITarget;
 	EDisplayColorGamut HDRDisplayColorGamut;
+	ESlatePostRT UsedSlatePostBuffers;
 	bool bWireFrame;
 	bool bIsHDR;
 
@@ -45,6 +47,7 @@ struct FSlateRenderingParams
 		, ViewRect(FIntRect())
 		, Time(InTime)
 		, HDRDisplayColorGamut(EDisplayColorGamut::sRGB_D65)
+		, UsedSlatePostBuffers(ESlatePostRT::None)
 		, bWireFrame(false)
 		, bIsHDR(false)
 	{
@@ -80,7 +83,7 @@ public:
 
 	virtual void FlushGeneratedResources();
 
-	void BlurRectExternal(FRHICommandListImmediate& RHICmdList, FTexture2DRHIRef BlurSrc, FTextureReferenceRHIRef& BlurDst, FIntPoint DstExtent, float BlurStrength) const;
+	void BlurRectExternal(FRHICommandListImmediate& RHICmdList, FRHITexture* BlurSrc, FRHITexture* BlurDst, FIntRect SrcRect, FIntRect DstRect, float BlurStrength) const;
 
 private:
 	ETextureSamplerFilter GetSamplerFilter(const UTexture* Texture) const;
diff --git a/Engine/Source/Runtime/SlateRHIRenderer/Private/SlateRHIRenderingPolicyInterface.cpp b/Engine/Source/Runtime/SlateRHIRenderer/Private/SlateRHIRenderingPolicyInterface.cpp
--- a/Engine/Source/Runtime/SlateRHIRenderer/Private/SlateRHIRenderingPolicyInterface.cpp
+++ b/Engine/Source/Runtime/SlateRHIRenderer/Private/SlateRHIRenderingPolicyInterface.cpp
@@ -2,12 +2,18 @@
 
 #include "Interfaces/SlateRHIRenderingPolicyInterface.h"
 #include "SlateRHIRenderingPolicy.h"
+#include "SlateRHIRenderer.h"
 
 FSlateRHIRenderingPolicyInterface::FSlateRHIRenderingPolicyInterface(FSlateRHIRenderingPolicy* InRenderingPolicy)
 	: RenderingPolicy(InRenderingPolicy)
 {
 }
 
+int32 FSlateRHIRenderingPolicyInterface::GetProcessSlatePostBuffers()
+{
+	return FSlateRHIRenderer::GetProcessSlatePostBuffers();
+}
+
 bool FSlateRHIRenderingPolicyInterface::IsValid() const
 {
 	return RenderingPolicy != nullptr;
@@ -33,10 +39,10 @@ bool FSlateRHIRenderingPolicyInterface::GetApplyColorDeficiencyCorrection() cons
 	return false;
 }
 
-void FSlateRHIRenderingPolicyInterface::BlurRectExternal(FRHICommandListImmediate& RHICmdList, FTexture2DRHIRef BlurSrc, FTextureReferenceRHIRef& BlurDst, FIntPoint DstExtent, float BlurStrength) const
+void FSlateRHIRenderingPolicyInterface::BlurRectExternal(FRHICommandListImmediate& RHICmdList, FRHITexture* BlurSrc, FRHITexture* BlurDst, FIntRect SrcRect, FIntRect DstRect, float BlurStrength) const
 {
 	if (RenderingPolicy)
 	{
-		RenderingPolicy->BlurRectExternal(RHICmdList, BlurSrc, BlurDst, DstExtent, BlurStrength);
+		RenderingPolicy->BlurRectExternal(RHICmdList, BlurSrc, BlurDst, SrcRect, DstRect, BlurStrength);
 	}
 }
\ No newline at end of file
diff --git a/Engine/Source/Runtime/SlateRHIRenderer/Public/FX/SlateFXSubsystem.h b/Engine/Source/Runtime/SlateRHIRenderer/Public/FX/SlateFXSubsystem.h
--- a/Engine/Source/Runtime/SlateRHIRenderer/Public/FX/SlateFXSubsystem.h
+++ b/Engine/Source/Runtime/SlateRHIRenderer/Public/FX/SlateFXSubsystem.h
@@ -5,10 +5,12 @@
 #include "Subsystems/EngineSubsystem.h"
 
 #include "Engine/World.h"
-#include "SlateRHIRendererSettings.h"
 
 #include "SlateFXSubsystem.generated.h"
 
+class FSlateRHIPostBufferProcessorProxy;
+class USlateRHIPostBufferProcessor;
+
 UCLASS(DisplayName = "Slate FX Subsystem")
 class SLATERHIRENDERER_API USlateFXSubsystem : public UEngineSubsystem
 {
@@ -16,11 +18,21 @@ class SLATERHIRENDERER_API USlateFXSubsystem : public UEngineSubsystem
 
 public:
 
-	//~Begin UGameInstanceSubsystem Insterface
+	static USlateRHIPostBufferProcessor* GetPostProcessor(ESlatePostRT InSlatePostBufferBit);
+	static TSharedPtr<FSlateRHIPostBufferProcessorProxy> GetPostProcessorProxy(ESlatePostRT InSlatePostBufferBit);
+
+	//~ Begin UObject Interface.
+	virtual void BeginDestroy() override;
+	//~ End UObject Interface.
+
+	//~Begin UGameInstanceSubsystem Interface
 	virtual bool ShouldCreateSubsystem(UObject* Outer) const override;
 	virtual void Initialize(FSubsystemCollectionBase& Collection) override;
 	virtual void Deinitialize() override;
-	//~End UGameInstanceSubsystem Insterface
+	//~End UGameInstanceSubsystem Interface
+
+	/** Get post processor proxy for a particular post buffer index, if it exists */
+	TSharedPtr<FSlateRHIPostBufferProcessorProxy> GetSlatePostProcessorProxy(ESlatePostRT InPostBufferBit);
 
 public:
 
@@ -41,4 +53,9 @@ private:
 
 	/** Callback to remove processors on world cleanup */
 	void OnPostWorldCleanup(UWorld* World, bool SessionEnded, bool bCleanupResources);
+
+private:
+
+	/** Map of post RT buffer index to buffer processor renderthread proxies, if they exist */
+	TMap<ESlatePostRT, TSharedPtr<FSlateRHIPostBufferProcessorProxy>> SlatePostBufferProcessorProxies;
 };
diff --git a/Engine/Source/Runtime/SlateRHIRenderer/Public/FX/SlatePostBufferBlur.h b/Engine/Source/Runtime/SlateRHIRenderer/Public/FX/SlatePostBufferBlur.h
--- a/Engine/Source/Runtime/SlateRHIRenderer/Public/FX/SlatePostBufferBlur.h
+++ b/Engine/Source/Runtime/SlateRHIRenderer/Public/FX/SlatePostBufferBlur.h
@@ -7,6 +7,31 @@
 
 #include "SlatePostBufferBlur.generated.h"
 
+
+/**
+ * Proxy for post buffer processor that the renderthread uses to perform processing
+ * This proxy exists because generally speaking usage on UObjects on the renderthread
+ * is a race condition due to UObjects being managed / updated by the game thread
+ */
+class SLATERHIRENDERER_API FSlatePostBufferBlurProxy : public FSlateRHIPostBufferProcessorProxy
+{
+
+public:
+
+	//~ Begin FSlateRHIPostBufferProcessorProxy Interface
+	virtual void PostProcess_Renderthread(FRHICommandListImmediate& RHICmdList, FRHITexture* Src, FRHITexture* Dst, FIntRect SrcRect, FIntRect DstRect, FSlateRHIRenderingPolicyInterface InRenderingPolicy) override;
+	virtual void OnUpdateValuesRenderThread() override;
+	//~ End FSlateRHIPostBufferProcessorProxy Interface
+
+protected:
+
+	/** Blur strength to use when processing, renderthread version actually used to draw. Must be updated via render command except during initialization. */
+	float GaussianBlurStrength_RenderThread = 10;
+
+	/** Fence to allow for us to queue only one update per draw command */
+	FRenderCommandFence ParamUpdateFence;
+};
+
 /**
  * Slate Post Buffer Processor that performs a simple gaussian blur to the backbuffer
  * 
@@ -24,5 +49,15 @@ public:
 
 public:
 
-	virtual void PostProcess(FRenderResource* InViewInfo, FRenderResource* InViewportTexture, FVector2D InElementWindowSize, FSlateRHIRenderingPolicyInterface InRenderingPolicy, UTextureRenderTarget2D* InSlatePostBuffer);
+	USlatePostBufferBlur();
+	virtual ~USlatePostBufferBlur() override;
+
+	//~ Begin USlateRHIPostBufferProcessor Interface
+	virtual void PostProcess(FRenderResource* InViewInfo, FRenderResource* InViewportTexture, FVector2D InElementWindowSize, FSlateRHIRenderingPolicyInterface InRenderingPolicy, UTextureRenderTarget2D* InSlatePostBuffer) override;
+	virtual TSharedPtr<FSlateRHIPostBufferProcessorProxy> GetRenderThreadProxy();
+	//~ End USlateRHIPostBufferProcessor Interface
+
+private:
+
+	TSharedPtr<FSlateRHIPostBufferProcessorProxy> RenderThreadProxy;
 };
\ No newline at end of file
diff --git a/Engine/Source/Runtime/SlateRHIRenderer/Public/FX/SlateRHIPostBufferProcessor.h b/Engine/Source/Runtime/SlateRHIRenderer/Public/FX/SlateRHIPostBufferProcessor.h
--- a/Engine/Source/Runtime/SlateRHIRenderer/Public/FX/SlateRHIPostBufferProcessor.h
+++ b/Engine/Source/Runtime/SlateRHIRenderer/Public/FX/SlateRHIPostBufferProcessor.h
@@ -10,14 +10,68 @@
 #include "Rendering/SlateRenderer.h"
 #include "RHIFwd.h"
 #include "UObject/Object.h"
+#include "UObject/ObjectMacros.h"
 #include "Interfaces/SlateRHIRenderingPolicyInterface.h"
 
 #include "SlateRHIPostBufferProcessor.generated.h"
 
+class USlateRHIPostBufferProcessor;
+
+/**
+ * Proxy for post buffer processor that the renderthread uses to perform processing
+ * This proxy exists because generally speaking usage on UObjects on the renderthread
+ * is a race condition due to UObjects being managed / updated by the game thread
+ */
+class SLATERHIRENDERER_API FSlateRHIPostBufferProcessorProxy : public TSharedFromThis<FSlateRHIPostBufferProcessorProxy>
+{
+
+public:
+
+	virtual ~FSlateRHIPostBufferProcessorProxy()
+	{
+	}
+
+	/**
+	 * Called directly inside renderthread to perform some processing, do not enque commands here as we are already in the renderthread
+	 *
+	 * @param Src					Source texture to process as Input
+	 * @param Dst					Destination texture to store process output
+	 * @param SrcRect				Rect within source texture to sample, in PIE this is a subsection, in standlone it should be the entire texture
+	 * @param DstRect				Rect within output to write out, since this is within the renderthread should almost always be Dst's Extent.
+	 * @param InRenderingPolicy		RenderingPolicy used to assist / perform processing
+	 */
+	virtual void PostProcess_Renderthread(FRHICommandListImmediate& RHICmdList, FRHITexture* Src, FRHITexture* Dst, FIntRect SrcRect, FIntRect DstRect, FSlateRHIRenderingPolicyInterface InRenderingPolicy)
+	{
+	}
+
+	/** 
+	 * Called when an post buffer update element is added to a renderbatch, 
+	 * gives proxies a chance to queue updates to their renderthread values based on the UObject processor.
+	 * These updates should likely be guarded by an 'FRenderCommandFence' to avoid duplicate updates
+	 */
+	virtual void OnUpdateValuesRenderThread()
+	{
+	}
+
+	/**
+	 * Set the UObject that we are a renderthread proxy for, useful for doing gamethread updates from the proxy
+	 */
+	void SetOwningProcessorObject(USlateRHIPostBufferProcessor* InParentObject)
+	{
+		ParentObject = InParentObject;
+	}
+
+protected:
+
+	/** Pointer to processor that we are a proxy for, external design constraints should ensure that this is always valid */
+	TWeakObjectPtr<USlateRHIPostBufferProcessor> ParentObject;
+};
+
 /**
  * Base class for types that can process the backbuffer scene into the slate post buffer.
  * 
- * Implement 'PostProcess' in your derived class. For an example see: USlatePostBufferBlur.
+ * Implement 'PostProcess' in your derived class. Additionally, you need to create a renderthread proxy that derives from 'FSlateRHIPostBufferProcessorProxy'
+ * For an example see: USlatePostBufferBlur.
  */
 UCLASS(Abstract, Blueprintable, CollapseCategories)
 class SLATERHIRENDERER_API USlateRHIPostBufferProcessor : public UObject
@@ -26,6 +80,10 @@ class SLATERHIRENDERER_API USlateRHIPostBufferProcessor : public UObject
 
 public:
 
+	virtual ~USlateRHIPostBufferProcessor()
+	{
+	}
+
 	/**
 	 * Overridable postprocess for the given source scene backbuffer provided in 'Src' into 'Dst'
 	 * You must override this method. In your override, you should copy params before executing 'ENQUEUE_RENDER_COMMAND'.
@@ -42,6 +100,14 @@ public:
 	{
 	}
 
+	/**
+	 * Gets proxy for this post buffer processor, for execution on the renderthread
+	 */
+	virtual TSharedPtr<FSlateRHIPostBufferProcessorProxy> GetRenderThreadProxy()
+	{
+		return nullptr;
+	}
+
 protected:
 
 	/** 
diff --git a/Engine/Source/Runtime/SlateRHIRenderer/Public/Interfaces/SlateRHIRenderingPolicyInterface.h b/Engine/Source/Runtime/SlateRHIRenderer/Public/Interfaces/SlateRHIRenderingPolicyInterface.h
--- a/Engine/Source/Runtime/SlateRHIRenderer/Public/Interfaces/SlateRHIRenderingPolicyInterface.h
+++ b/Engine/Source/Runtime/SlateRHIRenderer/Public/Interfaces/SlateRHIRenderingPolicyInterface.h
@@ -20,14 +20,42 @@ class SLATERHIRENDERER_API FSlateRHIRenderingPolicyInterface
 public:
 	FSlateRHIRenderingPolicyInterface(FSlateRHIRenderingPolicy* InRenderingPolicy);
 
+	static int32 GetProcessSlatePostBuffers();
+
 	bool IsValid() const;
 	bool IsVertexColorInLinearSpace() const;
 	bool GetApplyColorDeficiencyCorrection() const;
 
-	void BlurRectExternal(FRHICommandListImmediate& RHICmdList, FTexture2DRHIRef BlurSrc, FTextureReferenceRHIRef& BlurDst, FIntPoint DstExtent, float BlurStrength) const;
+	void BlurRectExternal(FRHICommandListImmediate& RHICmdList, FRHITexture* BlurSrc, FRHITexture* BlurDst, FIntRect SrcRect, FIntRect DstRect, float BlurStrength) const;
 
 private:
 
 	/** Rendering policy we are an interface to */
 	FSlateRHIRenderingPolicy* RenderingPolicy;
 };
+
+/**
+ * RHI version of custom slate element that can accept additional RHI params when performing draw
+ */
+class ICustomSlateElementRHI : public ICustomSlateElement
+{
+public:
+
+	/**
+	 * Called from the rendering thread when it is time to render the element
+	 *
+	 * @param RenderTarget				handle to the platform specific render target implementation.  Note this is already bound by Slate initially
+	 * @param Params					Params about current draw state
+	 * @param RenderingPolicyInterface	Interface to current rendering policy
+	 */
+	virtual void Draw_RHIRenderThread(class FRHICommandListImmediate& RHICmdList, const FTextureRHIRef& RenderTarget, const FSlateCustomDrawParams& Params, FSlateRHIRenderingPolicyInterface RenderingPolicyInterface)
+	{
+	}
+
+	//~ Begin ICustomSlateElement interface
+	virtual bool UsesAdditionalRHIParams() const override
+	{
+		return true;
+	}
+	//~ End ICustomSlateElement interface
+};
\ No newline at end of file
diff --git a/Engine/Source/Runtime/SlateRHIRenderer/Public/SlateRHIRendererSettings.h b/Engine/Source/Runtime/SlateRHIRenderer/Public/SlateRHIRendererSettings.h
--- a/Engine/Source/Runtime/SlateRHIRenderer/Public/SlateRHIRendererSettings.h
+++ b/Engine/Source/Runtime/SlateRHIRenderer/Public/SlateRHIRendererSettings.h
@@ -5,7 +5,6 @@
 #include "CoreMinimal.h"
 #include "Containers/ContainerAllocationPolicies.h"
 #include "Engine/DeveloperSettings.h"
-#include "FX/SlateRHIPostBufferProcessor.h"
 #include "Misc/EnumRange.h"
 #include "UObject/Object.h"
 #include "UObject/SoftObjectPath.h"
@@ -14,6 +13,8 @@
 
 #include "SlateRHIRendererSettings.generated.h"
 
+class USlateRHIPostBufferProcessor;
+
 /**
  * Settings for a particular Slate Post RT.
  * Notably if enabled & blur by default. To be updated with additional effects & to be expandable in game code / settings.
@@ -83,6 +84,10 @@ public:
 	USlateRHIRendererSettings();
 	~USlateRHIRendererSettings();
 
+	//~ Begin UObject Interface.
+	virtual void BeginDestroy() override;
+	//~ End UObject Interface.
+
 public:
 
 	/** Get settings struct for a particular post buffer index */
diff --git a/Engine/Source/Runtime/UMG/Private/Components/PostBufferUpdate.cpp b/Engine/Source/Runtime/UMG/Private/Components/PostBufferUpdate.cpp
new file mode 100644
--- /dev/null
+++ b/Engine/Source/Runtime/UMG/Private/Components/PostBufferUpdate.cpp
@@ -0,0 +1,70 @@
+// Copyright Epic Games, Inc. All Rights Reserved.
+
+#include "Components/PostBufferUpdate.h"
+
+#include "Slate/SPostBufferUpdate.h"
+
+#include UE_INLINE_GENERATED_CPP_BY_NAME(PostBufferUpdate)
+
+#define LOCTEXT_NAMESPACE "UMG"
+
+/////////////////////////////////////////////////////
+// UPostBufferUpdate
+
+UPostBufferUpdate::UPostBufferUpdate()
+{
+	bPerformDefaultPostBufferUpdate = true;
+	BuffersToUpdate = {};
+}
+
+void UPostBufferUpdate::SetPerformDefaultPostBufferUpdate(bool bInPerformDefaultPostBufferUpdate)
+{
+	if (bPerformDefaultPostBufferUpdate != bInPerformDefaultPostBufferUpdate)
+	{
+		bPerformDefaultPostBufferUpdate = bInPerformDefaultPostBufferUpdate;
+		if (MyPostBufferUpdate)
+		{
+			MyPostBufferUpdate->SetPerformDefaultPostBufferUpdate(bPerformDefaultPostBufferUpdate);
+		}
+	}
+}
+
+TSharedRef<SWidget> UPostBufferUpdate::RebuildWidget()
+{
+	MyPostBufferUpdate = SNew(SPostBufferUpdate)
+		.bPerformDefaultPostBufferUpdate(bPerformDefaultPostBufferUpdate);
+
+	MyPostBufferUpdate->SetBuffersToUpdate(BuffersToUpdate);
+
+	return MyPostBufferUpdate.ToSharedRef();
+}
+
+void UPostBufferUpdate::SynchronizeProperties()
+{
+	Super::SynchronizeProperties();
+
+	if (!MyPostBufferUpdate.IsValid())
+	{
+		return;
+	}
+
+	MyPostBufferUpdate->SetPerformDefaultPostBufferUpdate(bPerformDefaultPostBufferUpdate);
+	MyPostBufferUpdate->SetBuffersToUpdate(BuffersToUpdate);
+}
+
+UMG_API void UPostBufferUpdate::ReleaseSlateResources(bool bReleaseChildren)
+{
+	Super::ReleaseSlateResources(bReleaseChildren);
+
+	if (MyPostBufferUpdate.IsValid())
+	{
+		MyPostBufferUpdate->ReleasePostBufferUpdater();
+	}
+
+	MyPostBufferUpdate.Reset();
+}
+
+/////////////////////////////////////////////////////
+
+#undef LOCTEXT_NAMESPACE
+
diff --git a/Engine/Source/Runtime/UMG/Private/Slate/SPostBufferUpdate.cpp b/Engine/Source/Runtime/UMG/Private/Slate/SPostBufferUpdate.cpp
new file mode 100644
--- /dev/null
+++ b/Engine/Source/Runtime/UMG/Private/Slate/SPostBufferUpdate.cpp
@@ -0,0 +1,245 @@
+// Copyright Epic Games, Inc. All Rights Reserved.
+
+#include "Slate/SPostBufferUpdate.h"
+
+#include "Engine/AssetManager.h"
+#include "Engine/Engine.h"
+#include "Engine/TextureRenderTarget2D.h"
+
+// Exclude SlateRHIRenderer related includes & implementations from server builds since the module is not a dependency / will not link for UMG on the server
+#if !UE_SERVER
+#include "FX/SlateFXSubsystem.h"
+#include "FX/SlateRHIPostBufferProcessor.h"
+#include "Interfaces/SlateRHIRenderingPolicyInterface.h"
+#include "Rendering/DrawElements.h"
+#include "Rendering/ElementBatcher.h"
+#include "Math/MathFwd.h"
+#include "RenderUtils.h"
+#include "RHICommandList.h"
+#include "RHIFwd.h"
+#include "RHIResources.h"
+#include "RHIUtilities.h"
+#include "SlateRHIRendererSettings.h"
+
+/**
+ * Custom Slate drawer to update slate post buffer
+ */
+class FPostBufferUpdater : public ICustomSlateElementRHI
+{
+public:
+
+	//~ Begin ICustomSlateElementRHI interface
+	virtual void Draw_RHIRenderThread(FRHICommandListImmediate& RHICmdList, const FTextureRHIRef& InWindowBackBuffer, const FSlateCustomDrawParams& Params, class FSlateRHIRenderingPolicyInterface RenderingPolicyInterface) override;
+	virtual void PostCustomElementAdded(class FSlateElementBatcher& ElementBatcher) const override;
+	//~ End ICustomSlateElementRHI interface
+
+public:
+
+	/** True if we should perform the default post buffer update, used to set related state on ElementBatcher at Gamethread element batch time */
+	bool bPerformDefaultPostBufferUpdate = true;
+
+	/** True if buffers to update has been initialized. Used to ensure 'BuffersToUpdate_Renderthread' is only set once at initialization. */
+	bool bBuffersToUpdateInitialized = false;
+
+	/** 
+	 * Buffers that we should update, all of these buffers will be affected by 'bPerformDefaultPostBufferUpdate' if disabled 
+	 * This value is read by the Renderthread, so all non-initialization updates must be done via a render command.
+	 * See 'FSlatePostBufferBlurProxy::OnUpdateValuesRenderThread' as an example.
+	 * 
+	 * Additionally, note that this value should be masked against the currently enabled buffers in 'USlateRHIRendererSettings'
+	 */
+	ESlatePostRT BuffersToUpdate_Renderthread = ESlatePostRT::None;
+};
+
+/////////////////////////////////////////////////////
+// FPostBufferUpdater
+
+void FPostBufferUpdater::Draw_RHIRenderThread(FRHICommandListImmediate& RHICmdList, const FTextureRHIRef& InWindowBackBuffer, const FSlateCustomDrawParams& Params, FSlateRHIRenderingPolicyInterface RenderingPolicyInterface)
+{
+	if (FSlateRHIRenderingPolicyInterface::GetProcessSlatePostBuffers())
+	{
+		for (ESlatePostRT SlatePostBufferBit : TEnumRange<ESlatePostRT>())
+		{
+			bool bPostBufferBitUsed = (Params.UsedSlatePostBuffers & SlatePostBufferBit & BuffersToUpdate_Renderthread) != ESlatePostRT::None;
+
+			if (bPostBufferBitUsed)
+			{
+				UTextureRenderTarget2D* SlatePostBuffer = Cast<UTextureRenderTarget2D>(USlateRHIRendererSettings::Get()->TryGetPostBufferRT(SlatePostBufferBit));
+				if (!SlatePostBuffer)
+				{
+					continue;
+				}
+
+				FRHITexture* Src = InWindowBackBuffer.IsValid() ? InWindowBackBuffer.GetReference() : nullptr;
+				FRHITexture* Dst = SlatePostBuffer->TextureReference.TextureReferenceRHI;
+
+				if (!Src || !Dst)
+				{
+					continue;
+				}
+
+				FIntRect SrcRect = FIntRect(0, 0, Src->GetDesc().Extent.X, Src->GetDesc().Extent.Y);
+				FIntRect DstRect = FIntRect(0, 0, Dst->GetDesc().Extent.X, Dst->GetDesc().Extent.Y);
+
+				if (GIsEditor)
+				{
+					// In PIE, the backbuffer for slate will be the entire editor window, So instead use the ViewRect as our extent
+					SrcRect = Params.ViewRect;
+				}
+				
+				// Note: In PIE we draw the scene even during resizes, the ViewRect can change even after a previous draw
+				// leading to crashes if we do not ensure size, we skip below if sizes do not match. This can result in a PostRT 
+				// with no update in PIE (May appear white during active drag-resizing). However this is not an issue in 
+				// Standalone since we don't draw during active resizes.
+				if (SrcRect.Width() != DstRect.Width() || SrcRect.Height() != DstRect.Height())
+				{
+					continue;
+				}
+
+				if (TSharedPtr<FSlateRHIPostBufferProcessorProxy> PostProcessorProxy = USlateFXSubsystem::GetPostProcessorProxy(SlatePostBufferBit))
+				{
+					PostProcessorProxy->PostProcess_Renderthread(RHICmdList, Src, Dst, SrcRect, DstRect, RenderingPolicyInterface);
+				}
+				else
+				{
+					FRHICopyTextureInfo CopyInfo;
+
+					// Copy just the viewport RT if in PIE, else do entire backbuffer
+					CopyInfo.SourcePosition = FIntVector(SrcRect.Min.X, SrcRect.Min.Y, 0);
+					CopyInfo.Size = FIntVector(DstRect.Width(), DstRect.Height(), 1);
+					TransitionAndCopyTexture(RHICmdList, Src, Dst, CopyInfo);
+				}
+			}
+		}
+	}
+}
+
+void FPostBufferUpdater::PostCustomElementAdded(FSlateElementBatcher& ElementBatcher) const
+{
+	ESlatePostRT PrevResourceUpdatingPostBuffers = ElementBatcher.GetResourceUpdatingPostBuffers();
+	ElementBatcher.SetResourceUpdatingPostBuffers(PrevResourceUpdatingPostBuffers | BuffersToUpdate_Renderthread);
+
+	if (!bPerformDefaultPostBufferUpdate)
+	{
+		ESlatePostRT PrevSkipDefaultUpdatePostBuffers = ElementBatcher.GetSkipDefaultUpdatePostBuffers();
+		ElementBatcher.SetSkipDefaultUpdatePostBuffers(PrevSkipDefaultUpdatePostBuffers | BuffersToUpdate_Renderthread);
+	}
+
+	// Give proxies a chance to update their renderthread values.
+	if (FSlateRHIRenderingPolicyInterface::GetProcessSlatePostBuffers())
+	{
+		for (ESlatePostRT SlatePostBufferBit : TEnumRange<ESlatePostRT>())
+		{
+			if (!USlateRHIRendererSettings::Get()->GetSlatePostSetting(SlatePostBufferBit).bEnabled)
+			{
+				continue;
+			}
+
+			bool bPostBufferBitUsed = (SlatePostBufferBit & BuffersToUpdate_Renderthread) != ESlatePostRT::None;
+
+			if (bPostBufferBitUsed)
+			{
+				if (TSharedPtr<FSlateRHIPostBufferProcessorProxy> PostProcessorProxy = USlateFXSubsystem::GetPostProcessorProxy(SlatePostBufferBit))
+				{
+					PostProcessorProxy->OnUpdateValuesRenderThread();
+				}
+			}
+		}
+	}
+}
+
+#endif // !UE_SERVER
+
+/////////////////////////////////////////////////////
+// SPostBufferUpdate
+
+SLATE_IMPLEMENT_WIDGET(SPostBufferUpdate)
+void SPostBufferUpdate::PrivateRegisterAttributes(FSlateAttributeInitializer& AttributeInitializer)
+{
+}
+
+SPostBufferUpdate::SPostBufferUpdate()
+	: bPerformDefaultPostBufferUpdate(true)
+{
+}
+
+void SPostBufferUpdate::Construct(const FArguments& InArgs)
+{
+#if !UE_SERVER
+	bPerformDefaultPostBufferUpdate = InArgs._bPerformDefaultPostBufferUpdate;
+
+	PostBufferUpdater = MakeShared<FPostBufferUpdater>();
+	if (PostBufferUpdater)
+	{
+		PostBufferUpdater->bPerformDefaultPostBufferUpdate = bPerformDefaultPostBufferUpdate;
+	}
+#endif // !UE_SERVER
+}
+
+void SPostBufferUpdate::SetPerformDefaultPostBufferUpdate(bool bInPerformDefaultPostBufferUpdate)
+{
+#if !UE_SERVER
+	bPerformDefaultPostBufferUpdate = bInPerformDefaultPostBufferUpdate;
+
+	if (PostBufferUpdater)
+	{
+		PostBufferUpdater->bPerformDefaultPostBufferUpdate = bPerformDefaultPostBufferUpdate;
+	}
+#endif // !UE_SERVER
+}
+
+bool SPostBufferUpdate::GetPerformDefaultPostBufferUpdate() const
+{
+	return bPerformDefaultPostBufferUpdate;
+}
+
+void SPostBufferUpdate::SetBuffersToUpdate(const TArrayView<ESlatePostRT> InBuffersToUpdate)
+{
+#if !UE_SERVER
+	BuffersToUpdate = InBuffersToUpdate;
+
+	if (PostBufferUpdater && !PostBufferUpdater->bBuffersToUpdateInitialized)
+	{
+		PostBufferUpdater->BuffersToUpdate_Renderthread = ESlatePostRT::None;
+		for (ESlatePostRT BufferToUpdate : BuffersToUpdate)
+		{
+			if (USlateRHIRendererSettings::Get()->GetSlatePostSetting(BufferToUpdate).bEnabled)
+			{
+				PostBufferUpdater->BuffersToUpdate_Renderthread |= BufferToUpdate;
+			}
+		}
+
+		PostBufferUpdater->bBuffersToUpdateInitialized = true;
+	}
+#endif // !UE_SERVER
+}
+
+const TArrayView<const ESlatePostRT> SPostBufferUpdate::GetBuffersToUpdate() const
+{
+	return MakeArrayView(BuffersToUpdate);
+}
+
+UMG_API void SPostBufferUpdate::ReleasePostBufferUpdater()
+{
+#if !UE_SERVER
+	PostBufferUpdater.Reset();
+#endif // !UE_SERVER
+}
+
+int32 SPostBufferUpdate::OnPaint( const FPaintArgs& Args, const FGeometry& AllottedGeometry, const FSlateRect& MyCullingRect, FSlateWindowElementList& OutDrawElements, int32 LayerId, const FWidgetStyle& InWidgetStyle, bool bParentEnabled ) const
+{
+#if !UE_SERVER
+	FSlateRect RenderBoundingRect = AllottedGeometry.GetRenderBoundingRect();
+	FPaintGeometry PaintGeometry(RenderBoundingRect.GetTopLeft(), RenderBoundingRect.GetSize(), AllottedGeometry.GetAccumulatedLayoutTransform().GetScale());
+	FSlateDrawElement::MakeCustom(OutDrawElements, LayerId, PostBufferUpdater);
+#endif // !UE_SERVER
+
+	// Increment LayerId to ensure items afterwards are not processed
+	return ++LayerId;
+}
+
+FVector2D SPostBufferUpdate::ComputeDesiredSize( float ) const
+{
+	return FVector2D::ZeroVector;
+}
+
diff --git a/Engine/Source/Runtime/UMG/Public/Components/PostBufferUpdate.h b/Engine/Source/Runtime/UMG/Public/Components/PostBufferUpdate.h
new file mode 100644
--- /dev/null
+++ b/Engine/Source/Runtime/UMG/Public/Components/PostBufferUpdate.h
@@ -0,0 +1,56 @@
+// Copyright Epic Games, Inc. All Rights Reserved.
+
+#pragma once
+
+#include "CoreMinimal.h"
+#include "UObject/ObjectMacros.h"
+#include "Misc/Attribute.h"
+#include "Widgets/SWidget.h"
+#include "Components/Widget.h"
+#include "Rendering/SlateRendererTypes.h"
+
+#include "PostBufferUpdate.generated.h"
+
+class SPostBufferUpdate;
+
+/**
+ * Widget that when drawn, will trigger the slate post buffer to update. Does not draw anything itself.
+ * This allows for you to perform layered UI / sampling effects with the GetSlatePost material functions,
+ * by placing this widget after UI you would like to be processed / sampled is drawn.
+ *
+ * * No Children
+ */
+UCLASS(MinimalAPI)
+class UPostBufferUpdate : public UWidget
+{
+	GENERATED_BODY()
+
+	UPostBufferUpdate();
+
+private:
+
+	/** 
+	 * True if we should do the default post buffer update of the scene before any UI.
+	 * If any PostBufferUpdate widget has this set as false, be default scene copy / processing will not occur.
+	 */
+	UPROPERTY(EditAnywhere, BlueprintReadWrite, Setter="SetPerformDefaultPostBufferUpdate", Category = "Behavior", meta = (AllowPrivateAccess = "true"))
+	bool bPerformDefaultPostBufferUpdate;
+
+	/** Buffers that we should update, all of these buffers will be affected by 'bPerformDefaultPostBufferUpdate' if disabled */
+	UPROPERTY(EditAnywhere, Category = "Behavior", meta = (AllowPrivateAccess = "true"))
+	TArray<ESlatePostRT> BuffersToUpdate;
+
+public:
+	/** Set the orientation of the stack box. The existing elements will be rearranged. */
+	UMG_API void SetPerformDefaultPostBufferUpdate(bool bInPerformDefaultPostBufferUpdate);
+
+protected:
+	//~ Begin UWidget Interface
+	UMG_API virtual TSharedRef<SWidget> RebuildWidget() override;
+	UMG_API virtual void SynchronizeProperties() override;
+	UMG_API virtual void ReleaseSlateResources(bool bReleaseChildren) override;
+	//~ End UWidget Interface
+
+protected:
+	TSharedPtr<SPostBufferUpdate> MyPostBufferUpdate;
+};
diff --git a/Engine/Source/Runtime/UMG/Public/Slate/SPostBufferUpdate.h b/Engine/Source/Runtime/UMG/Public/Slate/SPostBufferUpdate.h
new file mode 100644
--- /dev/null
+++ b/Engine/Source/Runtime/UMG/Public/Slate/SPostBufferUpdate.h
@@ -0,0 +1,93 @@
+// Copyright Epic Games, Inc. All Rights Reserved.
+
+#pragma once
+
+#include "CoreMinimal.h"
+#include "Misc/Attribute.h"
+#include "Input/Reply.h"
+#include "Widgets/DeclarativeSyntaxSupport.h"
+#include "Widgets/SLeafWidget.h"
+#include "Styling/SlateTypes.h"
+
+class FPaintArgs;
+class FSlateWindowElementList;
+
+/**
+ * Custom Slate drawer to update slate post buffer
+ * 
+ * Note: Declared in .cpp to avoid UMG header dependencies on SlateRHIRenderer on Server
+ */
+class FPostBufferUpdater;
+
+/**
+ * Implements a widget that triggers a post buffer update on draw
+ */
+class SPostBufferUpdate
+	: public SLeafWidget
+{
+	SLATE_DECLARE_WIDGET_API(SPostBufferUpdate, SLeafWidget, UMG_API)
+
+public:
+	SLATE_BEGIN_ARGS( SPostBufferUpdate )
+		: _bPerformDefaultPostBufferUpdate( true )
+		{ }
+
+		/** True if we should perform the default post buffer update */
+		SLATE_ARGUMENT(bool, bPerformDefaultPostBufferUpdate)
+	SLATE_END_ARGS()
+
+	/** Constructor */
+	UMG_API SPostBufferUpdate();
+
+	/**
+	 * Construct this widget
+	 *
+	 * @param	InArgs	The declaration data for this widget
+	 */
+	UMG_API void Construct( const FArguments& InArgs );
+
+public:
+
+	/** Set if we should perform the default post buffer update */
+	UMG_API void SetPerformDefaultPostBufferUpdate(bool bInPerformDefaultPostBufferUpdate);
+
+	/** Get if we should perform the default post buffer update */
+	UMG_API bool GetPerformDefaultPostBufferUpdate() const;
+
+	/**
+	 * Set buffers to, this method only affects the 'FPostBufferUpdater' custom drawer once during initialization. 
+	 * Afterwards it will be a no-op due to the custom drawer having update buffers initialized.
+	 * This set-once behavior is done to avoid renderthread race conditions, & since you shouldn't need to modify
+	 * the buffers to update at runtime since we only draw if used anyway.
+	 */
+	UMG_API void SetBuffersToUpdate(const TArrayView<ESlatePostRT> InBuffersToUpdate);
+
+	/** Get buffers to update */
+	UMG_API const TArrayView<const ESlatePostRT> GetBuffersToUpdate() const;
+
+	/** Release Post Buffer Updater, called out of caution in case of reuse during 'ReleaseSlateResources' in UWidget */
+	UMG_API void ReleasePostBufferUpdater();
+
+public:
+
+	//~Begin Widget Interface
+	UMG_API virtual int32 OnPaint( const FPaintArgs& Args, const FGeometry& AllottedGeometry, const FSlateRect& MyCullingRect, FSlateWindowElementList& OutDrawElements, int32 LayerId, const FWidgetStyle& InWidgetStyle, bool bParentEnabled ) const override;
+	//~End Widget Interface
+
+protected:
+
+	//~Begin Widget Interface
+	UMG_API virtual FVector2D ComputeDesiredSize(float) const override;
+	//~End Widget Interface
+
+private:
+
+	/** True if we should perform the default post buffer update */
+	bool bPerformDefaultPostBufferUpdate;
+
+	/** Buffers that we should update, all of these buffers will be affected by 'bPerformDefaultPostBufferUpdate' if disabled */
+	TArray<ESlatePostRT> BuffersToUpdate;
+
+	/** Custom drawer used to trigger a post buffer update */
+	TSharedPtr<FPostBufferUpdater, ESPMode::ThreadSafe> PostBufferUpdater;
+};
diff --git a/Engine/Source/Runtime/UMG/UMG.Build.cs b/Engine/Source/Runtime/UMG/UMG.Build.cs
--- a/Engine/Source/Runtime/UMG/UMG.Build.cs
+++ b/Engine/Source/Runtime/UMG/UMG.Build.cs
@@ -45,6 +45,13 @@ public class UMG : ModuleRules
 			DynamicallyLoadedModuleNames.AddRange(
 				new string[] {
 					"ImageWrapper",
+				}
+			);
+
+			// Implementations depending on 'SlateRHIRenderer' should 
+			// be guarded with: '#if !UE_SERVER'.
+			PrivateDependencyModuleNames.AddRange(
+				new string[] {
 					"SlateRHIRenderer",
 				}
 			);
