Make SlateRHIRenderer return UTextureRenderTarget2D for post buffers instead of UObject (May fix buffers not loading on some build configs).
Make slate resources cache post buffer usage on first use.
Guard Post Buffer RT loading so that it only happens once to avoid any chance of log spam.

diff --git a/Engine/Source/Runtime/SlateRHIRenderer/Private/SlateMaterialResource.cpp b/Engine/Source/Runtime/SlateRHIRenderer/Private/SlateMaterialResource.cpp
--- a/Engine/Source/Runtime/SlateRHIRenderer/Private/SlateMaterialResource.cpp
+++ b/Engine/Source/Runtime/SlateRHIRenderer/Private/SlateMaterialResource.cpp
@@ -90,6 +90,28 @@ FSlateMaterialResource::FSlateMaterialResource(const UMaterialInterface& InMater
 	SlateProxy->ActualSize = InImageSize.IntPoint();
 	SlateProxy->Resource = this;
 
+	if (MaterialObject)
+	{
+		// Quality / Feature level irrelevant since flag to search all levels for both is true
+		TArray<UTexture*> OutUsedTextures;
+		MaterialObject->GetUsedTextures(OutUsedTextures, EMaterialQualityLevel::Num, true, ERHIFeatureLevel::Num, true);
+
+		CachedSlatePostBuffers = ESlatePostRT::None;
+		for (const UTexture* OutUsedTexture : OutUsedTextures)
+		{
+			for (const TPair<ESlatePostRT, FSlatePostSettings>& SlatePostSetting : USlateRHIRendererSettings::Get()->GetSlatePostSettings())
+			{
+				const ESlatePostRT SlatePostBitflag = SlatePostSetting.Key;
+				const FSlatePostSettings& SlatePostSettingValue = SlatePostSetting.Value;
+
+				if (SlatePostSettingValue.bEnabled && OutUsedTexture && OutUsedTexture->GetPathName() == SlatePostSettingValue.GetPathToSlatePostRT())
+				{
+					CachedSlatePostBuffers |= SlatePostBitflag;
+				}
+			}
+		}
+	}
+
 	if (MaterialProxy && (MaterialProxy->IsDeleted() || MaterialProxy->IsMarkedForGarbageCollection()))
 	{
 		MaterialProxy = nullptr;
diff --git a/Engine/Source/Runtime/SlateRHIRenderer/Private/SlateRHIRenderer.cpp b/Engine/Source/Runtime/SlateRHIRenderer/Private/SlateRHIRenderer.cpp
--- a/Engine/Source/Runtime/SlateRHIRenderer/Private/SlateRHIRenderer.cpp
+++ b/Engine/Source/Runtime/SlateRHIRenderer/Private/SlateRHIRenderer.cpp
@@ -1761,18 +1761,14 @@ void FSlateRHIRenderer::DrawWindows_Private(FSlateDrawBuffer& WindowDrawBuffer)
 					uint8 SlatePostBufferBitIndex = 0;
 					for (ESlatePostRT SlatePostBufferBit : TEnumRange<ESlatePostRT>())
 					{
-						if (!USlateRHIRendererSettings::Get()->GetSlatePostSetting(SlatePostBufferBit).bEnabled)
+						// We only attempt to load if the buffer is enabled, so just try to load / get the buffer
+						UTextureRenderTarget2D* SlatePostBuffer = USlateRHIRendererSettings::GetMutable()->LoadGetPostBufferRT(SlatePostBufferBit);
+						if (!SlatePostBuffer)
 						{
 							SlatePostBufferBitIndex++;
 							continue;
 						}
 
-						UTextureRenderTarget2D* SlatePostBuffer = Cast<UTextureRenderTarget2D>(USlateRHIRendererSettings::Get()->TryGetPostBufferRT(SlatePostBufferBit));
-						if (!SlatePostBuffer)
-						{
-							SlatePostBuffer = Cast<UTextureRenderTarget2D>(USlateRHIRendererSettings::GetMutable()->LoadGetPostBufferRT(SlatePostBufferBit));
-						}
-
 						bool bIsViewportPresentForPIE = GIsEditor ? Window->GetViewport().IsValid() : true;
 						bool bPostBufferBitUsed = (UsedSlatePostBufferBits & SlatePostBufferBit) != ESlatePostRT::None;
 						bool bSkipDefaultPostBufferUpdate = (SkipDefaultUpdatePostBufferBits & SlatePostBufferBit) != ESlatePostRT::None;
diff --git a/Engine/Source/Runtime/SlateRHIRenderer/Private/SlateRHIRendererSettings.cpp b/Engine/Source/Runtime/SlateRHIRenderer/Private/SlateRHIRendererSettings.cpp
--- a/Engine/Source/Runtime/SlateRHIRenderer/Private/SlateRHIRendererSettings.cpp
+++ b/Engine/Source/Runtime/SlateRHIRenderer/Private/SlateRHIRendererSettings.cpp
@@ -1,6 +1,7 @@
 // Copyright Epic Games, Inc. All Rights Reserved.
 
 #include "SlateRHIRendererSettings.h"
+#include "Engine/TextureRenderTarget2D.h"
 #include "FX/SlateRHIPostBufferProcessor.h"
 #include "HAL/IConsoleManager.h"
 
@@ -17,6 +18,7 @@ FSlatePostSettings::FSlatePostSettings()
 	, PostProcessorClass(nullptr)
 	, PathToSlatePostRT(FString())
 	, CachedSlatePostRT(nullptr)
+	, bLoadAttempted(false)
 {
 }
 
@@ -71,23 +73,32 @@ const FSlatePostSettings& USlateRHIRendererSettings::GetSlatePostSetting(ESlateP
 	return SlatePostSettings[InPostBufferBit];
 }
 
-UObject* USlateRHIRendererSettings::TryGetPostBufferRT(ESlatePostRT InPostBufferBit) const
+UTextureRenderTarget2D* USlateRHIRendererSettings::TryGetPostBufferRT(ESlatePostRT InPostBufferBit) const
 {
 	return SlatePostSettings[InPostBufferBit].CachedSlatePostRT;
 }
 
-UObject* USlateRHIRendererSettings::LoadGetPostBufferRT(ESlatePostRT InPostBufferBit)
+UTextureRenderTarget2D* USlateRHIRendererSettings::LoadGetPostBufferRT(ESlatePostRT InPostBufferBit)
 {
-	UObject* Result = TryGetPostBufferRT(InPostBufferBit);
+	FSlatePostSettings& SlatePostSetting = SlatePostSettings[InPostBufferBit];
 
-	if (!Result)
+	UTextureRenderTarget2D* Result = nullptr;
+
+	if (SlatePostSetting.bEnabled)
 	{
-		Result = LoadObject<UObject>(nullptr, *SlatePostSettings[InPostBufferBit].PathToSlatePostRT, nullptr, LOAD_None, nullptr);
+		Result = SlatePostSetting.CachedSlatePostRT;
 
-		if (Result)
+		if (!Result && !SlatePostSetting.bLoadAttempted)
 		{
-			Result->AddToRoot();
-			SlatePostSettings[InPostBufferBit].CachedSlatePostRT = Result;
+			Result = LoadObject<UTextureRenderTarget2D>(nullptr, *SlatePostSetting.PathToSlatePostRT, nullptr, LOAD_None, nullptr);
+
+			if (Result)
+			{
+				Result->AddToRoot();
+				SlatePostSetting.CachedSlatePostRT = Result;
+			}
+
+			SlatePostSetting.bLoadAttempted = true;
 		}
 	}
 
diff --git a/Engine/Source/Runtime/SlateRHIRenderer/Private/SlateUTextureResource.cpp b/Engine/Source/Runtime/SlateRHIRenderer/Private/SlateUTextureResource.cpp
--- a/Engine/Source/Runtime/SlateRHIRenderer/Private/SlateUTextureResource.cpp
+++ b/Engine/Source/Runtime/SlateRHIRenderer/Private/SlateUTextureResource.cpp
@@ -84,6 +84,18 @@ FSlateUTextureResource::FSlateUTextureResource(UTexture* InTexture)
 	{
 		Proxy->ActualSize = FIntPoint(InTexture->GetSurfaceWidth(), InTexture->GetSurfaceHeight());
 		Proxy->Resource = this;
+
+		CachedSlatePostBuffers = ESlatePostRT::None;
+		for (const TPair<ESlatePostRT, FSlatePostSettings>& SlatePostSetting : USlateRHIRendererSettings::Get()->GetSlatePostSettings())
+		{
+			const ESlatePostRT SlatePostBitflag = SlatePostSetting.Key;
+			const FSlatePostSettings& SlatePostSettingValue = SlatePostSetting.Value;
+
+			if (SlatePostSettingValue.bEnabled && InTexture && InTexture->GetPathName() == SlatePostSettingValue.GetPathToSlatePostRT())
+			{
+				CachedSlatePostBuffers |= SlatePostBitflag;
+			}
+		}
 	}
 }
 
diff --git a/Engine/Source/Runtime/SlateRHIRenderer/Public/SlateRHIRendererSettings.h b/Engine/Source/Runtime/SlateRHIRenderer/Public/SlateRHIRendererSettings.h
--- a/Engine/Source/Runtime/SlateRHIRenderer/Public/SlateRHIRendererSettings.h
+++ b/Engine/Source/Runtime/SlateRHIRenderer/Public/SlateRHIRendererSettings.h
@@ -14,6 +14,7 @@
 #include "SlateRHIRendererSettings.generated.h"
 
 class USlateRHIPostBufferProcessor;
+class UTextureRenderTarget2D;
 
 /**
  * Settings for a particular Slate Post RT.
@@ -56,7 +57,10 @@ private:
 
 	/** Cached load of Slate Post RT Asset */
 	UPROPERTY(Transient)
-	TObjectPtr<UObject> CachedSlatePostRT;
+	TObjectPtr<UTextureRenderTarget2D> CachedSlatePostRT;
+
+	/** True if we attempted to load the post RT asset already */
+	bool bLoadAttempted:1;
 };
 
 /**
@@ -101,10 +105,10 @@ public:
 public:
 
 	/** Try to get post RT asset, returns nullptr if not already loaded */
-	UObject* TryGetPostBufferRT(ESlatePostRT InPostBufferBit) const;
+	UTextureRenderTarget2D* TryGetPostBufferRT(ESlatePostRT InPostBufferBit) const;
 
 	/** Get post RT asset, loading if not already loaded */
-	UObject* LoadGetPostBufferRT(ESlatePostRT InPostBufferBit);
+	UTextureRenderTarget2D* LoadGetPostBufferRT(ESlatePostRT InPostBufferBit);
 
 	/** Get slate post settings map, non mutable */
 	const TMap<ESlatePostRT, FSlatePostSettings>& GetSlatePostSettings() const;
