Fix static analysis SlateUTextureResource using texture param before checking if null.

#tests Local Compile & Run

diff --git a/Engine/Source/Runtime/SlateRHIRenderer/Private/SlateUTextureResource.cpp b/Engine/Source/Runtime/SlateRHIRenderer/Private/SlateUTextureResource.cpp
--- a/Engine/Source/Runtime/SlateRHIRenderer/Private/SlateUTextureResource.cpp
+++ b/Engine/Source/Runtime/SlateRHIRenderer/Private/SlateUTextureResource.cpp
@@ -82,7 +82,10 @@ FSlateUTextureResource::FSlateUTextureResource(UTexture* InTexture)
 {
 	if(TextureObject)
 	{
-		Proxy->ActualSize = FIntPoint(InTexture->GetSurfaceWidth(), InTexture->GetSurfaceHeight());
+
+		Proxy->ActualSize = InTexture 
+			? FIntPoint(InTexture->GetSurfaceWidth(), InTexture->GetSurfaceHeight())
+			: FIntPoint(1, 1);
 		Proxy->Resource = this;
 
 		CachedSlatePostBuffers = ESlatePostRT::None;
