Add validity checks to slate material resource

#tests Compiled

diff --git a/Engine/Source/Runtime/SlateRHIRenderer/Private/SlateMaterialResource.cpp b/Engine/Source/Runtime/SlateRHIRenderer/Private/SlateMaterialResource.cpp
--- a/Engine/Source/Runtime/SlateRHIRenderer/Private/SlateMaterialResource.cpp
+++ b/Engine/Source/Runtime/SlateRHIRenderer/Private/SlateMaterialResource.cpp
@@ -144,21 +144,16 @@ bool FSlateMaterialResource::IsResourceValid() const
 
 void FSlateMaterialResource::UpdateMaterial(const UMaterialInterface& InMaterialResource, const FVector2f InImageSize, FSlateShaderResource* InTextureMask)
 {
+	MaterialObject = ensure(IsValid(&InMaterialResource) && InMaterialResource.IsValidLowLevelFast(false)) ? &InMaterialResource : nullptr;
+	MaterialProxy = MaterialObject ? MaterialObject->GetRenderProxy() : nullptr;
+
 #if SLATE_CHECK_UOBJECT_RENDER_RESOURCES
 	SlateMaterialResource::CheckInvalidUMaterial(InMaterialResource, DebugName);
 
-	MaterialObject = &InMaterialResource;
-	MaterialProxy = InMaterialResource.GetRenderProxy();
-
 	MaterialObjectWeakPtr = MaterialObject;
 	UpdateMaterialName();
 
 	SlateMaterialResource::CheckInvalidMaterialProxy(MaterialProxy, DebugName);
-
-#else
-
-	MaterialObject = &InMaterialResource;
-	MaterialProxy = InMaterialResource.GetRenderProxy();
 #endif
 
 	if (MaterialObject)
