Delete slate post process render targets when they are not used after (n) frame.

diff --git a/Engine/Source/Runtime/SlateRHIRenderer/Private/SlatePostProcessResource.cpp b/Engine/Source/Runtime/SlateRHIRenderer/Private/SlatePostProcessResource.cpp
--- a/Engine/Source/Runtime/SlateRHIRenderer/Private/SlatePostProcessResource.cpp
+++ b/Engine/Source/Runtime/SlateRHIRenderer/Private/SlatePostProcessResource.cpp
@@ -12,6 +12,7 @@ FSlatePostProcessResource::FSlatePostProcessResource(int32 InRenderTargetCount)
 	, RenderTargetSize(FIntPoint::ZeroValue)
 	, RenderTargetCount(InRenderTargetCount)
 {
+	FrameUsed = GFrameCounter;
 }
 
 FSlatePostProcessResource::~FSlatePostProcessResource()
@@ -31,6 +32,7 @@ void FSlatePostProcessResource::Update(const FIntPoint& NewSize, EPixelFormat Re
 		FIntPoint NewMaxSize(FMath::Max(NewSize.X, RenderTargetSize.X), FMath::Max(NewSize.Y, RenderTargetSize.Y));
 		ResizeTargets(NewMaxSize, RequestedPixelFormat);
 	}
+	FrameUsed = GFrameCounter;
 }
 
 void FSlatePostProcessResource::ResizeTargets(const FIntPoint& NewSize, EPixelFormat RequestedPixelFormat)
diff --git a/Engine/Source/Runtime/SlateRHIRenderer/Private/SlatePostProcessResource.h b/Engine/Source/Runtime/SlateRHIRenderer/Private/SlatePostProcessResource.h
--- a/Engine/Source/Runtime/SlateRHIRenderer/Private/SlatePostProcessResource.h
+++ b/Engine/Source/Runtime/SlateRHIRenderer/Private/SlatePostProcessResource.h
@@ -39,6 +39,8 @@ public:
 	virtual ESlateShaderResource::Type GetType() const override { return ESlateShaderResource::PostProcess; }
 
 	EPixelFormat GetPixelFormat() const { return PixelFormat; }
+	uint64 GetFrameUsed(){ return FrameUsed; }
+
 private:
 	/** Resizes targets to the new size */
 	void ResizeTargets(const FIntPoint& NewSize, EPixelFormat RequestedPixelFormat);
@@ -48,5 +50,6 @@ private:
 	EPixelFormat PixelFormat;
 	FIntPoint RenderTargetSize;
 	int32 RenderTargetCount;
+	uint64 FrameUsed;
 };
 
diff --git a/Engine/Source/Runtime/SlateRHIRenderer/Private/SlatePostProcessor.cpp b/Engine/Source/Runtime/SlateRHIRenderer/Private/SlatePostProcessor.cpp
--- a/Engine/Source/Runtime/SlateRHIRenderer/Private/SlatePostProcessor.cpp
+++ b/Engine/Source/Runtime/SlateRHIRenderer/Private/SlatePostProcessor.cpp
@@ -16,6 +16,13 @@
 DECLARE_CYCLE_STAT(TEXT("Slate PostProcessing RT"), STAT_SlatePostProcessingRTTime, STATGROUP_Slate);
 DECLARE_CYCLE_STAT(TEXT("Slate ColorDeficiency RT"), STAT_SlateColorDeficiencyRTTime, STATGROUP_Slate);
 
+int GSlateEnableDeleteUnusedPostProcess = 100;
+static FAutoConsoleVariableRef CVarSlateEnableDeleteUnusedPostProcess(
+	TEXT("Slate.EnableDeleteUnusedPostProcess"),
+	GSlateEnableDeleteUnusedPostProcess,
+	TEXT("Greater than zero implies that post process render targets will be deleted when they are not used after n frames.")
+);
+
 static const int32 NumIntermediateTargets = 2;
 
 FSlatePostProcessResource* FindSlatePostProcessResource(TArray<FSlatePostProcessResource*>& IntermediateTargetsArray, EPixelFormat PixelFormat)
@@ -547,6 +554,23 @@ void FSlatePostProcessor::ReleaseRenderTargets()
 	}
 }
 
+void FSlatePostProcessor::TickPostProcessResources()
+{
+	check(IsInGameThread());
+	
+	if (GSlateEnableDeleteUnusedPostProcess > 0)
+	{
+		for (TArray<FSlatePostProcessResource*>::TIterator It = IntermediateTargetsArray.CreateIterator(); It; ++It)
+		{
+			if (GFrameCounter - (*It)->GetFrameUsed() > GSlateEnableDeleteUnusedPostProcess)
+			{
+				(*It)->CleanUp();
+				It.RemoveCurrent();
+			}
+		}
+	}
+}
+
 void FSlatePostProcessor::DownsampleRect(FRHICommandListImmediate& RHICmdList, IRendererModule& RendererModule, const FPostProcessRectParams& Params, const FIntPoint& DownsampleSize, FSlatePostProcessResource* IntermediateTargets)
 {
 	SCOPED_DRAW_EVENT(RHICmdList, SlatePostProcessDownsample);
diff --git a/Engine/Source/Runtime/SlateRHIRenderer/Private/SlatePostProcessor.h b/Engine/Source/Runtime/SlateRHIRenderer/Private/SlatePostProcessor.h
--- a/Engine/Source/Runtime/SlateRHIRenderer/Private/SlatePostProcessor.h
+++ b/Engine/Source/Runtime/SlateRHIRenderer/Private/SlatePostProcessor.h
@@ -51,6 +51,7 @@ public:
 	void ColorDeficiency(FRHICommandListImmediate& RHICmdList, IRendererModule& RendererModule, const FPostProcessRectParams& RectParams);
 	
 	void ReleaseRenderTargets();
+	void TickPostProcessResources();
 
 private:
 	void DownsampleRect(FRHICommandListImmediate& RHICmdList, IRendererModule& RendererModule, const FPostProcessRectParams& Params, const FIntPoint& DownsampleSize, FSlatePostProcessResource* IntermediateTargets);
diff --git a/Engine/Source/Runtime/SlateRHIRenderer/Private/SlateRHIRenderer.cpp b/Engine/Source/Runtime/SlateRHIRenderer/Private/SlateRHIRenderer.cpp
--- a/Engine/Source/Runtime/SlateRHIRenderer/Private/SlateRHIRenderer.cpp
+++ b/Engine/Source/Runtime/SlateRHIRenderer/Private/SlateRHIRenderer.cpp
@@ -2416,6 +2416,9 @@ void FSlateRHIRenderer::ReleaseAccessedResources(bool bImmediatelyFlush)
 
 		//FlushCommands();
 	}
+
+	RenderingPolicy->TickPostProcessResources();
+
 }
 
 void FSlateRHIRenderer::RequestResize(const TSharedPtr<SWindow>& Window, uint32 NewWidth, uint32 NewHeight)
diff --git a/Engine/Source/Runtime/SlateRHIRenderer/Private/SlateRHIRenderingPolicy.cpp b/Engine/Source/Runtime/SlateRHIRenderer/Private/SlateRHIRenderingPolicy.cpp
--- a/Engine/Source/Runtime/SlateRHIRenderer/Private/SlateRHIRenderingPolicy.cpp
+++ b/Engine/Source/Runtime/SlateRHIRenderer/Private/SlateRHIRenderingPolicy.cpp
@@ -1672,6 +1672,11 @@ void FSlateRHIRenderingPolicy::FlushGeneratedResources()
 	PostProcessor->ReleaseRenderTargets();
 }
 
+void FSlateRHIRenderingPolicy::TickPostProcessResources()
+{
+	PostProcessor->TickPostProcessResources();
+}
+
 void FSlateRHIRenderingPolicy::BlurRectExternal(FRHICommandListImmediate& RHICmdList, FRHITexture* BlurSrc, FRHITexture* BlurDst, FIntRect SrcRect, FIntRect DstRect, float BlurStrength) const
 {
 	SLATE_DRAW_EVENT(RHICmdList, PostProcess);
diff --git a/Engine/Source/Runtime/SlateRHIRenderer/Private/SlateRHIRenderingPolicy.h b/Engine/Source/Runtime/SlateRHIRenderer/Private/SlateRHIRenderingPolicy.h
--- a/Engine/Source/Runtime/SlateRHIRenderer/Private/SlateRHIRenderingPolicy.h
+++ b/Engine/Source/Runtime/SlateRHIRenderer/Private/SlateRHIRenderingPolicy.h
@@ -76,6 +76,8 @@ public:
 	void SetUseGammaCorrection( bool bInUseGammaCorrection ) { bGammaCorrect = bInUseGammaCorrection; }
 	void SetApplyColorDeficiencyCorrection(bool bInApplyColorCorrection) { bApplyColorDeficiencyCorrection = bInApplyColorCorrection; }
 
+	void TickPostProcessResources();
+	
 	bool GetApplyColorDeficiencyCorrection() const { return bApplyColorDeficiencyCorrection; }
 
 	virtual void AddSceneAt(FSceneInterface* Scene, int32 Index) override;
