Fix Post Buffer Update PIE flickering when multiple windows not using post buffers was present
---
Remove various post buffer editor uses
Initialize some post buffer properties in ctor that previously were not initialized

diff --git a/Engine/Source/Runtime/SlateRHIRenderer/Private/FX/SlateFXSubsystem.cpp b/Engine/Source/Runtime/SlateRHIRenderer/Private/FX/SlateFXSubsystem.cpp
--- a/Engine/Source/Runtime/SlateRHIRenderer/Private/FX/SlateFXSubsystem.cpp
+++ b/Engine/Source/Runtime/SlateRHIRenderer/Private/FX/SlateFXSubsystem.cpp
@@ -3,6 +3,7 @@
 #include "FX/SlateFXSubsystem.h"
 
 #include "Engine/Engine.h"
+#include "Engine/EngineTypes.h"
 #include "Engine/GameInstance.h"
 #include "FX/SlateRHIPostBufferProcessor.h"
 #include "SlateRHIRendererSettings.h"
@@ -96,6 +97,11 @@ USlateRHIPostBufferProcessor* USlateFXSubsystem::GetSlatePostProcessor(ESlatePos
 
 void USlateFXSubsystem::OnPreWorldInitialization(UWorld* World, const UWorld::InitializationValues IVS)
 {
+	if (World && World->WorldType == EWorldType::EditorPreview)
+	{
+		return;
+	}
+
 	SlatePostBufferProcessors.Empty();
 	SlatePostBufferProcessorProxies.Empty();
 
@@ -118,6 +124,11 @@ void USlateFXSubsystem::OnPreWorldInitialization(UWorld* World, const UWorld::In
 
 void USlateFXSubsystem::OnPostWorldCleanup(UWorld* World, bool SessionEnded, bool bCleanupResources)
 {
+	if (World && World->WorldType == EWorldType::EditorPreview)
+	{
+		return;
+	}
+
 	SlatePostBufferProcessors.Empty();
 	SlatePostBufferProcessorProxies.Empty();
 }
diff --git a/Engine/Source/Runtime/SlateRHIRenderer/Private/SlateRHIRenderer.cpp b/Engine/Source/Runtime/SlateRHIRenderer/Private/SlateRHIRenderer.cpp
--- a/Engine/Source/Runtime/SlateRHIRenderer/Private/SlateRHIRenderer.cpp
+++ b/Engine/Source/Runtime/SlateRHIRenderer/Private/SlateRHIRenderer.cpp
@@ -244,6 +244,11 @@ FSlateRHIRenderer::FSlateRHIRenderer(TSharedRef<FSlateFontServices> InSlateFontS
 		FPlane(0, 0, 1, 0),
 		FPlane(0, 0, 0, 1));
 
+	for (uint64& LastFramePostBufferUsed : LastFramesPostBufferUsed)
+	{
+		LastFramePostBufferUsed = 0;
+	}
+
 	bTakingAScreenShot = false;
 	OutScreenshotData = NULL;
 	OutHDRScreenshotData = NULL;
@@ -1844,11 +1849,16 @@ void FSlateRHIRenderer::DrawWindows_Private(FSlateDrawBuffer& WindowDrawBuffer)
 								}
 
 								SlatePostRTFences[SlatePostBufferBitIndex].BeginFence();
+								LastFramesPostBufferUsed[SlatePostBufferBitIndex] = GFrameCounter;
 							}
 
 							bShrinkPostBufferRequested &= ~SlatePostBufferBit;
 						}
-						else if (SlatePostBuffer && SlatePostBuffer->GetResource() && SlatePostRTFences[SlatePostBufferBitIndex].IsFenceComplete() && (SlatePostBuffer->SizeX != 1 || SlatePostBuffer->SizeY != 1))
+						else if (SlatePostBuffer 
+							&& SlatePostBuffer->GetResource()
+							&& LastFramesPostBufferUsed[SlatePostBufferBitIndex] < GFrameCounter
+							&& SlatePostRTFences[SlatePostBufferBitIndex].IsFenceComplete() 
+							&& (SlatePostBuffer->SizeX != 1 || SlatePostBuffer->SizeY != 1))
 						{
 							if ((bShrinkPostBufferRequested & SlatePostBufferBit) == ESlatePostRT::None)
 							{
@@ -1933,6 +1943,7 @@ void FSlateRHIRenderer::DrawWindows_Private(FSlateDrawBuffer& WindowDrawBuffer)
 								if (bResourceUpdatedPostBuffer)
 								{
 									SlatePostRTFences[SlatePostBufferBitIndex].BeginFence();
+									LastFramesPostBufferUsed[SlatePostBufferBitIndex] = GFrameCounter;
 									bShrinkPostBufferRequested &= ~SlatePostBufferBit;
 								}
 
diff --git a/Engine/Source/Runtime/SlateRHIRenderer/Private/SlateRHIRenderer.h b/Engine/Source/Runtime/SlateRHIRenderer/Private/SlateRHIRenderer.h
--- a/Engine/Source/Runtime/SlateRHIRenderer/Private/SlateRHIRenderer.h
+++ b/Engine/Source/Runtime/SlateRHIRenderer/Private/SlateRHIRenderer.h
@@ -314,6 +314,7 @@ private:
 	bool bTakingAScreenShot;
 	bool bUpdateHDRDisplayInformation;
 	ESlatePostRT bShrinkPostBufferRequested;
+	uint64 LastFramesPostBufferUsed[(uint8)ESlatePostRT::Num];
 	FRenderCommandFence SlatePostRTFences[(uint8)ESlatePostRT::Num];
 	FIntRect ScreenshotRect;
 	FViewportInfo* ScreenshotViewportInfo;
diff --git a/Engine/Source/Runtime/UMG/Private/Components/PostBufferUpdate.cpp b/Engine/Source/Runtime/UMG/Private/Components/PostBufferUpdate.cpp
--- a/Engine/Source/Runtime/UMG/Private/Components/PostBufferUpdate.cpp
+++ b/Engine/Source/Runtime/UMG/Private/Components/PostBufferUpdate.cpp
@@ -34,7 +34,19 @@ TSharedRef<SWidget> UPostBufferUpdate::RebuildWidget()
 	MyPostBufferUpdate = SNew(SPostBufferUpdate)
 		.bPerformDefaultPostBufferUpdate(bPerformDefaultPostBufferUpdate);
 
-	MyPostBufferUpdate->SetBuffersToUpdate(BuffersToUpdate);
+	bool bSetBuffersToUpdate = true;
+
+#if WITH_EDITOR
+	if (IsDesignTime())
+	{
+		bSetBuffersToUpdate = false;
+	}
+#endif // WITH_EDITOR
+
+	if (bSetBuffersToUpdate)
+	{
+		MyPostBufferUpdate->SetBuffersToUpdate(BuffersToUpdate);
+	}
 
 	return MyPostBufferUpdate.ToSharedRef();
 }
@@ -49,7 +61,20 @@ void UPostBufferUpdate::SynchronizeProperties()
 	}
 
 	MyPostBufferUpdate->SetPerformDefaultPostBufferUpdate(bPerformDefaultPostBufferUpdate);
-	MyPostBufferUpdate->SetBuffersToUpdate(BuffersToUpdate);
+
+	bool bSetBuffersToUpdate = true;
+
+#if WITH_EDITOR
+	if (IsDesignTime())
+	{
+		bSetBuffersToUpdate = false;
+	}
+#endif // WITH_EDITOR
+
+	if (bSetBuffersToUpdate)
+	{
+		MyPostBufferUpdate->SetBuffersToUpdate(BuffersToUpdate);
+	}
 }
 
 UMG_API void UPostBufferUpdate::ReleaseSlateResources(bool bReleaseChildren)
diff --git a/Engine/Source/Runtime/UMG/Private/Slate/SPostBufferUpdate.cpp b/Engine/Source/Runtime/UMG/Private/Slate/SPostBufferUpdate.cpp
--- a/Engine/Source/Runtime/UMG/Private/Slate/SPostBufferUpdate.cpp
+++ b/Engine/Source/Runtime/UMG/Private/Slate/SPostBufferUpdate.cpp
@@ -160,6 +160,8 @@ void SPostBufferUpdate::PrivateRegisterAttributes(FSlateAttributeInitializer& At
 
 SPostBufferUpdate::SPostBufferUpdate()
 	: bPerformDefaultPostBufferUpdate(true)
+	, BuffersToUpdate({})
+	, PostBufferUpdater(nullptr)
 {
 }
 
@@ -168,6 +170,8 @@ void SPostBufferUpdate::Construct(const FArguments& InArgs)
 #if !UE_SERVER
 	bPerformDefaultPostBufferUpdate = InArgs._bPerformDefaultPostBufferUpdate;
 
+	BuffersToUpdate = {};
+
 	PostBufferUpdater = MakeShared<FPostBufferUpdater>();
 	if (PostBufferUpdater)
 	{
