Fix crash when deleting SlatePostProcessResources.

Add a null check on the iterator of the array containing SlatePostProcess pointers. If we call BeginReleaseresources before calling tick, the array can contain null pointers.

diff --git a/Engine/Source/Runtime/SlateRHIRenderer/Private/SlatePostProcessor.cpp b/Engine/Source/Runtime/SlateRHIRenderer/Private/SlatePostProcessor.cpp
--- a/Engine/Source/Runtime/SlateRHIRenderer/Private/SlatePostProcessor.cpp
+++ b/Engine/Source/Runtime/SlateRHIRenderer/Private/SlatePostProcessor.cpp
@@ -562,9 +562,10 @@ void FSlatePostProcessor::TickPostProcessResources()
 	{
 		for (TArray<FSlatePostProcessResource*>::TIterator It = IntermediateTargetsArray.CreateIterator(); It; ++It)
 		{
-			if (GFrameCounter - (*It)->GetFrameUsed() > GSlateEnableDeleteUnusedPostProcess)
+			FSlatePostProcessResource* SlatePostProcessResource = (*It);
+			if (SlatePostProcessResource != nullptr && (GFrameCounter - SlatePostProcessResource->GetFrameUsed() > GSlateEnableDeleteUnusedPostProcess))
 			{
-				(*It)->CleanUp();
+				SlatePostProcessResource->CleanUp();
 				It.RemoveCurrent();
 			}
 		}
