Slate: Change the key of SlateResourceManager.

diff --git a/Engine/Source/Runtime/SlateRHIRenderer/Private/SlateRHIResourceManager.cpp b/Engine/Source/Runtime/SlateRHIRenderer/Private/SlateRHIResourceManager.cpp
--- a/Engine/Source/Runtime/SlateRHIRenderer/Private/SlateRHIResourceManager.cpp
+++ b/Engine/Source/Runtime/SlateRHIRenderer/Private/SlateRHIResourceManager.cpp
@@ -234,7 +234,7 @@ void FDynamicResourceMap::RemoveExpiredMaterialResources(TArray< TSharedPtr<FSla
 	for (FMaterialResourceMap::TIterator It(MaterialMap); It; ++It)
 	{
 		FMaterialKey& Key = It.Key();
-		if (!Key.Material.IsValid())
+		if (Key.Material.ResolveObjectPtr() == nullptr)
 		{
 			RemovedMaterials.Push(It.Value());
 			It.Value()->ResetMaterial();
diff --git a/Engine/Source/Runtime/SlateRHIRenderer/Private/SlateRHIResourceManager.h b/Engine/Source/Runtime/SlateRHIRenderer/Private/SlateRHIResourceManager.h
--- a/Engine/Source/Runtime/SlateRHIRenderer/Private/SlateRHIResourceManager.h
+++ b/Engine/Source/Runtime/SlateRHIRenderer/Private/SlateRHIResourceManager.h
@@ -6,6 +6,7 @@
 #include "Stats/Stats.h"
 #include "Textures/TextureAtlas.h"
 #include "UObject/GCObject.h"
+#include "UObject/ObjectKey.h"
 #include "Containers/Queue.h"
 #include "Rendering/ShaderResourceManager.h"
 #include "Rendering/RenderingCommon.h"
@@ -30,7 +31,7 @@ class FSlateVectorGraphicsCache;
  */
 struct FMaterialKey
 {
-	TWeakObjectPtr<const UMaterialInterface> Material;
+	FObjectKey Material;
 	const FVector2f ImageSize;
 	int32 MaskKey;
 
