Fix rounded border not cropping with materials

diff --git a/Engine/Source/Runtime/SlateRHIRenderer/Private/SlateMaterialShader.cpp b/Engine/Source/Runtime/SlateRHIRenderer/Private/SlateMaterialShader.cpp
--- a/Engine/Source/Runtime/SlateRHIRenderer/Private/SlateMaterialShader.cpp
+++ b/Engine/Source/Runtime/SlateRHIRenderer/Private/SlateMaterialShader.cpp
@@ -151,5 +151,6 @@ IMPLEMENT_SLATE_MATERIALSHADER_TYPE(Default);
 IMPLEMENT_SLATE_MATERIALSHADER_TYPE(Border);
 IMPLEMENT_SLATE_MATERIALSHADER_TYPE(GrayscaleFont);
 IMPLEMENT_SLATE_MATERIALSHADER_TYPE(ColorFont);
+IMPLEMENT_SLATE_MATERIALSHADER_TYPE(RoundedBox);
 IMPLEMENT_SLATE_MATERIALSHADER_TYPE(SdfFont);
 IMPLEMENT_SLATE_MATERIALSHADER_TYPE(MsdfFont);
diff --git a/Engine/Source/Runtime/SlateRHIRenderer/Private/SlateRHIRenderingPolicy.cpp b/Engine/Source/Runtime/SlateRHIRenderer/Private/SlateRHIRenderingPolicy.cpp
--- a/Engine/Source/Runtime/SlateRHIRenderer/Private/SlateRHIRenderingPolicy.cpp
+++ b/Engine/Source/Runtime/SlateRHIRenderer/Private/SlateRHIRenderingPolicy.cpp
@@ -1605,8 +1605,6 @@ void FSlateRHIRenderingPolicy::ChooseMaterialShaderTypes(ESlateShader ShaderType
 {
 	switch (ShaderType)
 	{
-	case ESlateShader::RoundedBox:
-		// Todo rounded box not supported in materials currently, intentional fall-through to Default
 	case ESlateShader::Default:
 		OutShaderTypes.AddShaderType<TSlateMaterialShaderPS<ESlateShader::Default>>();
 		break;
@@ -1622,6 +1620,9 @@ void FSlateRHIRenderingPolicy::ChooseMaterialShaderTypes(ESlateShader ShaderType
 	case ESlateShader::Custom:
 		OutShaderTypes.AddShaderType<TSlateMaterialShaderPS<ESlateShader::Custom>>();
 		break;
+	case ESlateShader::RoundedBox:
+		OutShaderTypes.AddShaderType<TSlateMaterialShaderPS<ESlateShader::RoundedBox>>();
+		break;
 	case ESlateShader::SdfFont:
 		OutShaderTypes.AddShaderType<TSlateMaterialShaderPS<ESlateShader::SdfFont>>();
 		break;
