Emit material names when drawing

diff --git a/Engine/Source/Runtime/SlateRHIRenderer/Private/SlateRHIRenderingPolicy.cpp b/Engine/Source/Runtime/SlateRHIRenderer/Private/SlateRHIRenderingPolicy.cpp
--- a/Engine/Source/Runtime/SlateRHIRenderer/Private/SlateRHIRenderingPolicy.cpp
+++ b/Engine/Source/Runtime/SlateRHIRenderer/Private/SlateRHIRenderingPolicy.cpp
@@ -58,8 +58,10 @@ static FAutoConsoleVariableRef CVarSlateDrawBatchNum(TEXT("Slate.DrawBatchNum"),
 
 #if !(UE_BUILD_SHIPPING || UE_BUILD_TEST)
 	#define SLATE_DRAW_EVENT(RHICmdList, EventName) SCOPED_CONDITIONAL_DRAW_EVENT(RHICmdList, EventName, SlateEnableDrawEvents);
+	#define SLATE_DRAW_EVENTF(RHICmdList, EventName, Format, ...) SCOPED_CONDITIONAL_DRAW_EVENTF(RHICmdList, EventName, SlateEnableDrawEvents, Format, ##__VA_ARGS__);
 #else
 	#define SLATE_DRAW_EVENT(RHICmdList, EventName)
+	#define SLATE_DRAW_EVENTF(RHICmdList, EventName, Format, ...);
 #endif
 
 TAutoConsoleVariable<int32> CVarSlateAbsoluteIndices(
@@ -1073,7 +1075,6 @@ void FSlateRHIRenderingPolicy::DrawElements(
 			}
 			else if (GEngine && ShaderResource && ShaderResource->GetType() == ESlateShaderResource::Material && ShaderType != ESlateShader::PostProcess)
 			{
-				SLATE_DRAW_EVENT(RHICmdList, MaterialBatch);
 				check(RHICmdList.IsInsideRenderPass());
 
 				check(RenderBatch.NumIndices > 0);
@@ -1110,6 +1111,8 @@ void FSlateRHIRenderingPolicy::DrawElements(
 				FSlateMaterialResource* MaterialShaderResource = (FSlateMaterialResource*)ShaderResource;
 				if (const FMaterialRenderProxy* MaterialRenderProxy = MaterialShaderResource->GetRenderProxy())
 				{
+					SLATE_DRAW_EVENTF(RHICmdList, MaterialBatch, TEXT("Slate Material: %s"), *MaterialRenderProxy->GetMaterialName());
+
 					MaterialShaderResource->CheckForStaleResources();
 
 					const bool bUseInstancing = RenderBatch.InstanceCount > 0 && RenderBatch.InstanceData != nullptr;
