Delete slate post process render targets when they are not used after (n) frame. Optimization was previously rollback because the cleanup of the render targets array was done on the gamethread. Since the array is manipulated on the render thread (FindSlatePostProcessResource), the fix was to move the TickPostProcessResources call after the rendering of the slate elements on the render thread.

Disable the optim for now on all platforms. We will enable it later, when we enable another feature.

diff --git a/Engine/Source/Runtime/SlateRHIRenderer/Private/SlatePostProcessResource.cpp b/Engine/Source/Runtime/SlateRHIRenderer/Private/SlatePostProcessResource.cpp
--- a/Engine/Source/Runtime/SlateRHIRenderer/Private/SlatePostProcessResource.cpp
+++ b/Engine/Source/Runtime/SlateRHIRenderer/Private/SlatePostProcessResource.cpp
@@ -12,6 +12,7 @@ FSlatePostProcessResource::FSlatePostProcessResource(int32 InRenderTargetCount)
 	, RenderTargetSize(FIntPoint::ZeroValue)
 	, RenderTargetCount(InRenderTargetCount)
 {
+	FrameUsed = GFrameCounter;
 }
 
 FSlatePostProcessResource::~FSlatePostProcessResource()
@@ -31,6 +32,7 @@ void FSlatePostProcessResource::Update(const FIntPoint& NewSize, EPixelFormat Re
 		FIntPoint NewMaxSize(FMath::Max(NewSize.X, RenderTargetSize.X), FMath::Max(NewSize.Y, RenderTargetSize.Y));
 		ResizeTargets(NewMaxSize, RequestedPixelFormat);
 	}
+	FrameUsed = GFrameCounter;
 }
 
 void FSlatePostProcessResource::ResizeTargets(const FIntPoint& NewSize, EPixelFormat RequestedPixelFormat)
diff --git a/Engine/Source/Runtime/SlateRHIRenderer/Private/SlatePostProcessResource.h b/Engine/Source/Runtime/SlateRHIRenderer/Private/SlatePostProcessResource.h
--- a/Engine/Source/Runtime/SlateRHIRenderer/Private/SlatePostProcessResource.h
+++ b/Engine/Source/Runtime/SlateRHIRenderer/Private/SlatePostProcessResource.h
@@ -39,6 +39,8 @@ public:
 	virtual ESlateShaderResource::Type GetType() const override { return ESlateShaderResource::PostProcess; }
 
 	EPixelFormat GetPixelFormat() const { return PixelFormat; }
+	uint64 GetFrameUsed() { return FrameUsed; }
+
 private:
 	/** Resizes targets to the new size */
 	void ResizeTargets(const FIntPoint& NewSize, EPixelFormat RequestedPixelFormat);
@@ -48,5 +50,6 @@ private:
 	EPixelFormat PixelFormat;
 	FIntPoint RenderTargetSize;
 	int32 RenderTargetCount;
+	uint64 FrameUsed;
 };
 
diff --git a/Engine/Source/Runtime/SlateRHIRenderer/Private/SlatePostProcessor.cpp b/Engine/Source/Runtime/SlateRHIRenderer/Private/SlatePostProcessor.cpp
--- a/Engine/Source/Runtime/SlateRHIRenderer/Private/SlatePostProcessor.cpp
+++ b/Engine/Source/Runtime/SlateRHIRenderer/Private/SlatePostProcessor.cpp
@@ -16,6 +16,13 @@
 DECLARE_CYCLE_STAT(TEXT("Slate PostProcessing RT"), STAT_SlatePostProcessingRTTime, STATGROUP_Slate);
 DECLARE_CYCLE_STAT(TEXT("Slate ColorDeficiency RT"), STAT_SlateColorDeficiencyRTTime, STATGROUP_Slate);
 
+int GSlateEnableDeleteUnusedPostProcess = 0;
+static FAutoConsoleVariableRef CVarSlateEnableDeleteUnusedPostProcess(
+	TEXT("Slate.EnableDeleteUnusedPostProcess"),
+	GSlateEnableDeleteUnusedPostProcess,
+	TEXT("Greater than zero implies that post process render targets will be deleted when they are not used after n frames.")
+);
+
 static const int32 NumIntermediateTargets = 2;
 
 FSlatePostProcessResource* FindSlatePostProcessResource(TArray<FSlatePostProcessResource*>& IntermediateTargetsArray, EPixelFormat PixelFormat)
@@ -844,3 +851,20 @@ int32 FSlatePostProcessor::ComputeBlurWeights(int32 KernelSize, float StdDev, TA
 {
 	return ComputeWeights(KernelSize, StdDev, OutWeightsAndOffsets);
 }
+
+void FSlatePostProcessor::TickPostProcessResources()
+{
+	if (GSlateEnableDeleteUnusedPostProcess > 0)
+	{
+		check(IsInRenderingThread());
+
+		for (TArray<FSlatePostProcessResource*>::TIterator It = IntermediateTargetsArray.CreateIterator(); It; ++It)
+		{
+			if (GFrameCounter - (*It)->GetFrameUsed() > GSlateEnableDeleteUnusedPostProcess)
+			{
+				(*It)->CleanUp();
+				It.RemoveCurrent();
+			}
+		}
+	}
+}
\ No newline at end of file
diff --git a/Engine/Source/Runtime/SlateRHIRenderer/Private/SlatePostProcessor.h b/Engine/Source/Runtime/SlateRHIRenderer/Private/SlatePostProcessor.h
--- a/Engine/Source/Runtime/SlateRHIRenderer/Private/SlatePostProcessor.h
+++ b/Engine/Source/Runtime/SlateRHIRenderer/Private/SlatePostProcessor.h
@@ -51,6 +51,7 @@ public:
 	void ColorDeficiency(FRHICommandListImmediate& RHICmdList, IRendererModule& RendererModule, const FPostProcessRectParams& RectParams);
 	
 	void ReleaseRenderTargets();
+	void TickPostProcessResources();
 
 private:
 	void DownsampleRect(FRHICommandListImmediate& RHICmdList, IRendererModule& RendererModule, const FPostProcessRectParams& Params, const FIntPoint& DownsampleSize, FSlatePostProcessResource* IntermediateTargets);
diff --git a/Engine/Source/Runtime/SlateRHIRenderer/Private/SlateRHIRenderer.cpp b/Engine/Source/Runtime/SlateRHIRenderer/Private/SlateRHIRenderer.cpp
--- a/Engine/Source/Runtime/SlateRHIRenderer/Private/SlateRHIRenderer.cpp
+++ b/Engine/Source/Runtime/SlateRHIRenderer/Private/SlateRHIRenderer.cpp
@@ -1540,6 +1540,9 @@ void FSlateRHIRenderer::DrawWindow_RenderThread(FRHICommandListImmediate& RHICmd
 		ScreenshotViewportInfo = nullptr;
 	}
 
+	// check if we need to cleanup slate render targets alloc.
+	RenderingPolicy->TickPostProcessResources();
+
 	// Calculate renderthread time (excluding idle time).	
 	uint32 StartTime = FPlatformTime::Cycles();
 
diff --git a/Engine/Source/Runtime/SlateRHIRenderer/Private/SlateRHIRenderingPolicy.cpp b/Engine/Source/Runtime/SlateRHIRenderer/Private/SlateRHIRenderingPolicy.cpp
--- a/Engine/Source/Runtime/SlateRHIRenderer/Private/SlateRHIRenderingPolicy.cpp
+++ b/Engine/Source/Runtime/SlateRHIRenderer/Private/SlateRHIRenderingPolicy.cpp
@@ -1676,6 +1676,11 @@ void FSlateRHIRenderingPolicy::FlushGeneratedResources()
 	PostProcessor->ReleaseRenderTargets();
 }
 
+void FSlateRHIRenderingPolicy::TickPostProcessResources()
+{
+	PostProcessor->TickPostProcessResources();
+}
+
 void FSlateRHIRenderingPolicy::BlurRectExternal(FRHICommandListImmediate& RHICmdList, FRHITexture* BlurSrc, FRHITexture* BlurDst, FIntRect SrcRect, FIntRect DstRect, float BlurStrength) const
 {
 	SLATE_DRAW_EVENT(RHICmdList, PostProcess);
diff --git a/Engine/Source/Runtime/SlateRHIRenderer/Private/SlateRHIRenderingPolicy.h b/Engine/Source/Runtime/SlateRHIRenderer/Private/SlateRHIRenderingPolicy.h
--- a/Engine/Source/Runtime/SlateRHIRenderer/Private/SlateRHIRenderingPolicy.h
+++ b/Engine/Source/Runtime/SlateRHIRenderer/Private/SlateRHIRenderingPolicy.h
@@ -76,6 +76,8 @@ public:
 	void SetUseGammaCorrection( bool bInUseGammaCorrection ) { bGammaCorrect = bInUseGammaCorrection; }
 	void SetApplyColorDeficiencyCorrection(bool bInApplyColorCorrection) { bApplyColorDeficiencyCorrection = bInApplyColorCorrection; }
 
+	void TickPostProcessResources();
+
 	bool GetApplyColorDeficiencyCorrection() const { return bApplyColorDeficiencyCorrection; }
 
 	virtual void AddSceneAt(FSceneInterface* Scene, int32 Index) override;
