Removed unused static (causing issues with editor ASan builds due to link cycle).

diff --git a/Engine/Source/Runtime/SlateRHIRenderer/Private/SlateAtlasedTextureResource.cpp b/Engine/Source/Runtime/SlateRHIRenderer/Private/SlateAtlasedTextureResource.cpp
--- a/Engine/Source/Runtime/SlateRHIRenderer/Private/SlateAtlasedTextureResource.cpp
+++ b/Engine/Source/Runtime/SlateRHIRenderer/Private/SlateAtlasedTextureResource.cpp
@@ -3,8 +3,6 @@
 #include "SlateAtlasedTextureResource.h"
 #include "Slate/SlateTextureAtlasInterface.h"
 
-TSharedPtr<FSlateAtlasedTextureResource> FSlateAtlasedTextureResource::NullResource = MakeShareable( new FSlateAtlasedTextureResource(nullptr) );
-
 FSlateAtlasedTextureResource::FSlateAtlasedTextureResource(UTexture* InTexture)
 	: FSlateBaseUTextureResource(InTexture)
 {
diff --git a/Engine/Source/Runtime/SlateRHIRenderer/Private/SlateAtlasedTextureResource.h b/Engine/Source/Runtime/SlateRHIRenderer/Private/SlateAtlasedTextureResource.h
--- a/Engine/Source/Runtime/SlateRHIRenderer/Private/SlateAtlasedTextureResource.h
+++ b/Engine/Source/Runtime/SlateRHIRenderer/Private/SlateAtlasedTextureResource.h
@@ -14,8 +14,6 @@ struct FSlateAtlasData;
 class FSlateAtlasedTextureResource : public FSlateBaseUTextureResource
 {
 public:
-	static TSharedPtr<FSlateAtlasedTextureResource> NullResource;
-
 	/** Initializes a new atlased UTexture resource, the incoming texture should be the entire atlas. */
 	FSlateAtlasedTextureResource(UTexture* InTexture);
 	
diff --git a/Engine/Source/Runtime/SlateRHIRenderer/Private/SlateNativeTextureResource.cpp b/Engine/Source/Runtime/SlateRHIRenderer/Private/SlateNativeTextureResource.cpp
--- a/Engine/Source/Runtime/SlateRHIRenderer/Private/SlateNativeTextureResource.cpp
+++ b/Engine/Source/Runtime/SlateRHIRenderer/Private/SlateNativeTextureResource.cpp
@@ -4,8 +4,6 @@
 #include "Textures/SlateShaderResource.h"
 #include "Slate/SlateTextures.h"
 
-TSharedPtr<FSlateDynamicTextureResource> FSlateDynamicTextureResource::NullResource = MakeShareable( new FSlateDynamicTextureResource( NULL ) );
-
 FSlateDynamicTextureResource::FSlateDynamicTextureResource(FSlateTexture2DRHIRef* ExistingTexture)
 	: Proxy(new FSlateShaderResourceProxy)
 	, RHIRefTexture(ExistingTexture != NULL ? ExistingTexture : new FSlateTexture2DRHIRef(NULL, 0, 0))
@@ -13,7 +11,6 @@ FSlateDynamicTextureResource::FSlateDynamicTextureResource(FSlateTexture2DRHIRef
 	Proxy->Resource = RHIRefTexture;
 }
 
-
 FSlateDynamicTextureResource::~FSlateDynamicTextureResource()
 {
 	if(Proxy)
diff --git a/Engine/Source/Runtime/SlateRHIRenderer/Private/SlateNativeTextureResource.h b/Engine/Source/Runtime/SlateRHIRenderer/Private/SlateNativeTextureResource.h
--- a/Engine/Source/Runtime/SlateRHIRenderer/Private/SlateNativeTextureResource.h
+++ b/Engine/Source/Runtime/SlateRHIRenderer/Private/SlateNativeTextureResource.h
@@ -11,8 +11,6 @@ class FSlateTexture2DRHIRef;
 class FSlateDynamicTextureResource
 {
 public:
-	static TSharedPtr<FSlateDynamicTextureResource> NullResource;
-
 	FSlateDynamicTextureResource(FSlateTexture2DRHIRef* ExistingTexture);
 	~FSlateDynamicTextureResource();
 
