Force Linear Format on post buffers as some platforms do not support sRGB 10_10_10.

diff --git a/Engine/Source/Runtime/SlateRHIRenderer/Private/SlateRHIRenderer.cpp b/Engine/Source/Runtime/SlateRHIRenderer/Private/SlateRHIRenderer.cpp
--- a/Engine/Source/Runtime/SlateRHIRenderer/Private/SlateRHIRenderer.cpp
+++ b/Engine/Source/Runtime/SlateRHIRenderer/Private/SlateRHIRenderer.cpp
@@ -1822,7 +1822,7 @@ void FSlateRHIRenderer::DrawWindows_Private(FSlateDrawBuffer& WindowDrawBuffer)
 								: SlatePostBuffer->GetFormat() == EPixelFormat::PF_A2B10G10R10;
 							if (SlatePostBuffer->SizeX != SizeSlatePostRT.X || SlatePostBuffer->SizeY != SizeSlatePostRT.Y || !bIsPixelFormatCorrect)
 							{
-								SlatePostBuffer->InitCustomFormat(SizeSlatePostRT.X, SizeSlatePostRT.Y, bHDREnabled ? EPixelFormat::PF_FloatRGBA : EPixelFormat::PF_A2B10G10R10, false);
+								SlatePostBuffer->InitCustomFormat(SizeSlatePostRT.X, SizeSlatePostRT.Y, bHDREnabled ? EPixelFormat::PF_FloatRGBA : EPixelFormat::PF_A2B10G10R10, true);
 							}
 
 							const FVector2D ElementWindowSize = ElementList.GetWindowSize();
@@ -1888,7 +1888,7 @@ void FSlateRHIRenderer::DrawWindows_Private(FSlateDrawBuffer& WindowDrawBuffer)
 							else
 							{
 								// Resize unused SlatePostRTs to 1x1.
-								SlatePostBuffer->InitCustomFormat(1, 1, IsHDREnabled() ? EPixelFormat::PF_FloatRGBA : EPixelFormat::PF_A2B10G10R10, false);
+								SlatePostBuffer->InitCustomFormat(1, 1, IsHDREnabled() ? EPixelFormat::PF_FloatRGBA : EPixelFormat::PF_A2B10G10R10, true);
 								bShrinkPostBufferRequested &= ~SlatePostBufferBit;
 							}
 						}
