Disable Clear Mask on slate post buffer RT, as it causes issues when sampled on certain platforms.

#tests PS5 Cooked Frontend no longer repro

diff --git a/Engine/Source/Runtime/Engine/Classes/Engine/TextureRenderTarget2D.h b/Engine/Source/Runtime/Engine/Classes/Engine/TextureRenderTarget2D.h
--- a/Engine/Source/Runtime/Engine/Classes/Engine/TextureRenderTarget2D.h
+++ b/Engine/Source/Runtime/Engine/Classes/Engine/TextureRenderTarget2D.h
@@ -128,6 +128,9 @@ class UTextureRenderTarget2D : public UTextureRenderTarget
 	UPROPERTY()
 	uint8 bForceLinearGamma:1;
 
+	/** If true, render target will be marked as a resolve target. */
+	uint8 bIsResolveTarget:1;
+
 	/** Whether to support storing HDR values, which requires more memory. */
 	UPROPERTY()
 	uint8 bHDR_DEPRECATED:1;
diff --git a/Engine/Source/Runtime/Engine/Private/TextureRenderTarget2D.cpp b/Engine/Source/Runtime/Engine/Private/TextureRenderTarget2D.cpp
--- a/Engine/Source/Runtime/Engine/Private/TextureRenderTarget2D.cpp
+++ b/Engine/Source/Runtime/Engine/Private/TextureRenderTarget2D.cpp
@@ -44,6 +44,7 @@ UTextureRenderTarget2D::UTextureRenderTarget2D(const FObjectInitializer& ObjectI
 	ClearColor = FLinearColor(0.0f, 0.0f, 0.0f, 1.0f);
 	OverrideFormat = PF_Unknown;
 	bForceLinearGamma = true; // <<-- if you set RTF_RGBA8_SRGB, this is turned off
+	bIsResolveTarget = false;
 	MipsSamplerFilter = Filter;
 	MipsAddressU = TA_Clamp;
 	MipsAddressV = TA_Clamp;
@@ -507,6 +508,11 @@ ETextureCreateFlags FTextureRenderTarget2DResource::GetCreateFlags()
 		TexCreateFlags |= ETextureCreateFlags::UAV;
 	}
 
+	if (Owner->bIsResolveTarget)
+	{
+		TexCreateFlags |= ETextureCreateFlags::ResolveTargetable;
+	}
+
 	return TexCreateFlags | ETextureCreateFlags::RenderTargetable | ETextureCreateFlags::ShaderResource;
 }
 
diff --git a/Engine/Source/Runtime/SlateRHIRenderer/Private/SlateRHIRenderer.cpp b/Engine/Source/Runtime/SlateRHIRenderer/Private/SlateRHIRenderer.cpp
--- a/Engine/Source/Runtime/SlateRHIRenderer/Private/SlateRHIRenderer.cpp
+++ b/Engine/Source/Runtime/SlateRHIRenderer/Private/SlateRHIRenderer.cpp
@@ -1822,6 +1822,7 @@ void FSlateRHIRenderer::DrawWindows_Private(FSlateDrawBuffer& WindowDrawBuffer)
 								: SlatePostBuffer->GetFormat() == EPixelFormat::PF_A2B10G10R10;
 							if (SlatePostBuffer->SizeX != SizeSlatePostRT.X || SlatePostBuffer->SizeY != SizeSlatePostRT.Y || !bIsPixelFormatCorrect)
 							{
+								SlatePostBuffer->bIsResolveTarget = true;
 								SlatePostBuffer->InitCustomFormat(SizeSlatePostRT.X, SizeSlatePostRT.Y, bHDREnabled ? EPixelFormat::PF_FloatRGBA : EPixelFormat::PF_A2B10G10R10, true);
 							}
 
@@ -1888,6 +1889,7 @@ void FSlateRHIRenderer::DrawWindows_Private(FSlateDrawBuffer& WindowDrawBuffer)
 							else
 							{
 								// Resize unused SlatePostRTs to 1x1.
+								SlatePostBuffer->bIsResolveTarget = true;
 								SlatePostBuffer->InitCustomFormat(1, 1, IsHDREnabled() ? EPixelFormat::PF_FloatRGBA : EPixelFormat::PF_A2B10G10R10, true);
 								bShrinkPostBufferRequested &= ~SlatePostBufferBit;
 							}
