Add missing transtition to SRV to slate blur external

#tests Console cooked verify fixed & that DCC compression was used

#localization none

diff --git a/Engine/Source/Runtime/SlateRHIRenderer/Private/SlateRHIRenderer.cpp b/Engine/Source/Runtime/SlateRHIRenderer/Private/SlateRHIRenderer.cpp
--- a/Engine/Source/Runtime/SlateRHIRenderer/Private/SlateRHIRenderer.cpp
+++ b/Engine/Source/Runtime/SlateRHIRenderer/Private/SlateRHIRenderer.cpp
@@ -1822,7 +1822,6 @@ void FSlateRHIRenderer::DrawWindows_Private(FSlateDrawBuffer& WindowDrawBuffer)
 								: SlatePostBuffer->GetFormat() == EPixelFormat::PF_A2B10G10R10;
 							if (SlatePostBuffer->SizeX != SizeSlatePostRT.X || SlatePostBuffer->SizeY != SizeSlatePostRT.Y || !bIsPixelFormatCorrect)
 							{
-								SlatePostBuffer->bNoFastClear = true;
 								SlatePostBuffer->InitCustomFormat(SizeSlatePostRT.X, SizeSlatePostRT.Y, bHDREnabled ? EPixelFormat::PF_FloatRGBA : EPixelFormat::PF_A2B10G10R10, true);
 							}
 
@@ -1889,7 +1888,6 @@ void FSlateRHIRenderer::DrawWindows_Private(FSlateDrawBuffer& WindowDrawBuffer)
 							else
 							{
 								// Resize unused SlatePostRTs to 1x1.
-								SlatePostBuffer->bNoFastClear = true;
 								SlatePostBuffer->InitCustomFormat(1, 1, IsHDREnabled() ? EPixelFormat::PF_FloatRGBA : EPixelFormat::PF_A2B10G10R10, true);
 								bShrinkPostBufferRequested &= ~SlatePostBufferBit;
 							}
diff --git a/Engine/Source/Runtime/SlateRHIRenderer/Private/SlateRHIRenderingPolicy.cpp b/Engine/Source/Runtime/SlateRHIRenderer/Private/SlateRHIRenderingPolicy.cpp
--- a/Engine/Source/Runtime/SlateRHIRenderer/Private/SlateRHIRenderingPolicy.cpp
+++ b/Engine/Source/Runtime/SlateRHIRenderer/Private/SlateRHIRenderingPolicy.cpp
@@ -1739,5 +1739,7 @@ void FSlateRHIRenderingPolicy::BlurRectExternal(FRHICommandListImmediate& RHICmd
 
 	IRendererModule& RendererModule = FModuleManager::GetModuleChecked<IRendererModule>(RendererModuleName);
 	PostProcessor->BlurRect(RHICmdList, RendererModule, BlurParams, RectParams);
+
+	RHICmdList.Transition(FRHITransitionInfo(BlurDst, ERHIAccess::RTV, ERHIAccess::SRVGraphics));
 }
 	
