Slate3D - remove unnecessary evaluate all cached uniform expressions

diff --git a/Engine/Source/Runtime/SlateRHIRenderer/Private/Slate3DRenderer.cpp b/Engine/Source/Runtime/SlateRHIRenderer/Private/Slate3DRenderer.cpp
--- a/Engine/Source/Runtime/SlateRHIRenderer/Private/Slate3DRenderer.cpp
+++ b/Engine/Source/Runtime/SlateRHIRenderer/Private/Slate3DRenderer.cpp
@@ -143,14 +143,8 @@ void FSlate3DRenderer::DrawWindowToTarget_RenderThread(FRHICommandListImmediate&
 
 	checkSlow(Context.RenderTarget);
 
-	//Update cached uniforms to avoid an ensure being thrown due to a null shader map (UE-110263)
-	//Same fix on UE5 stream - CL16165057 
-	FMaterialRenderProxy::UpdateDeferredCachedUniformExpressions();
-
 	const TArray<TSharedRef<FSlateWindowElementList>>& WindowsToDraw = Context.WindowDrawBuffer->GetWindowElementLists();
 
-	FMaterialRenderProxy::UpdateDeferredCachedUniformExpressions();
-
 	// Enqueue a command to unlock the draw buffer after all windows have been drawn
 	RenderTargetPolicy->BeginDrawingWindows();
 
diff --git a/Engine/Source/Runtime/SlateRHIRenderer/Private/SlateRHIRenderingPolicy.cpp b/Engine/Source/Runtime/SlateRHIRenderer/Private/SlateRHIRenderingPolicy.cpp
--- a/Engine/Source/Runtime/SlateRHIRenderer/Private/SlateRHIRenderingPolicy.cpp
+++ b/Engine/Source/Runtime/SlateRHIRenderer/Private/SlateRHIRenderingPolicy.cpp
@@ -1083,6 +1083,9 @@ void FSlateRHIRenderingPolicy::DrawElements(
 				{
 					SLATE_DRAW_EVENTF(RHICmdList, MaterialBatch, TEXT("Slate Material: %s"), MaterialRenderProxy->GetMaterialName());
 
+					const ERHIFeatureLevel::Type ViewFeatureLevel = ActiveSceneView.GetFeatureLevel();
+					// Update Material in case of null ShaderMap
+					MaterialRenderProxy->UpdateUniformExpressionCacheIfNeeded(RHICmdList, ViewFeatureLevel);
 					MaterialShaderResource->CheckForStaleResources();
 
 					const bool bUseInstancing = RenderBatch.InstanceCount > 0 && RenderBatch.InstanceData != nullptr;
@@ -1094,7 +1097,6 @@ void FSlateRHIRenderingPolicy::DrawElements(
 					ChooseMaterialShaderTypes(ShaderType, bUseInstancing, ShaderTypesToGet);
 					const FMaterial* EffectiveMaterial = nullptr;
 
-					const ERHIFeatureLevel::Type ViewFeatureLevel = ActiveSceneView.GetFeatureLevel();
 					while(MaterialRenderProxy)
 					{
 						const FMaterial* Material = MaterialRenderProxy->GetMaterialNoFallback(ViewFeatureLevel);
