Always update cursor position uniform so we can use it in UI materials

#tests PIE / EUW

diff --git a/Engine/Source/Runtime/SlateRHIRenderer/Private/SlateRHIRenderer.cpp b/Engine/Source/Runtime/SlateRHIRenderer/Private/SlateRHIRenderer.cpp
--- a/Engine/Source/Runtime/SlateRHIRenderer/Private/SlateRHIRenderer.cpp
+++ b/Engine/Source/Runtime/SlateRHIRenderer/Private/SlateRHIRenderer.cpp
@@ -140,6 +140,7 @@ struct FSlateDrawWindowCommandParams
 	FSlateRHIRenderer* Renderer;
 	FSlateWindowElementList* WindowElementList;
 	SWindow* Window;
+	FIntPoint CursorPostion;
 	FIntRect ViewRect;
 	ESlatePostRT UsedSlatePostBuffers;
 #if WANTS_DRAW_MESH_EVENTS
@@ -1021,6 +1022,7 @@ void RenderSlateBatch(FTextureRHIRef SlateRenderTarget, bool bClear, bool bIsHDR
 				RenderParams.bIsHDR = bIsHDR;
 				RenderParams.HDRDisplayColorGamut = ViewportInfo.HDRDisplayColorGamut;
 				RenderParams.ViewRect = DrawCommandParams.ViewRect;
+				RenderParams.CursorPostion = DrawCommandParams.CursorPostion;
 				RenderParams.UsedSlatePostBuffers = DrawCommandParams.UsedSlatePostBuffers;
 				if (ViewportInfo.bSceneHDREnabled && !bIsHDR)
 				{
@@ -1924,10 +1926,22 @@ void FSlateRHIRenderer::DrawWindows_Private(FSlateDrawBuffer& WindowDrawBuffer)
 
 					FSlateDrawWindowCommandParams Params;
 
+					FVector2f CursorPosition = FSlateApplication::Get().GetCursorPos();
+
+					CursorPosition -= Window->GetPositionInScreen();
+					if (TSharedPtr<ISlateViewport> Viewport = Window->GetViewport())
+					{
+						if (TSharedPtr<SWidget> ViewportWidget = Viewport->GetWidget().Pin())
+						{
+							CursorPosition -= ViewportWidget->GetPaintSpaceGeometry().AbsolutePosition;
+						}
+					}
+
 					Params.Renderer = this;
 					Params.WindowElementList = &ElementList;
 					Params.Window = Window;
 					Params.ViewRect = GetViewRect();
+					Params.CursorPostion = FIntPoint(CursorPosition.X, CursorPosition.Y);
 					Params.UsedSlatePostBuffers = UsedSlatePostBufferBits;
 #if WANTS_DRAW_MESH_EVENTS
 					Params.WindowTitle = Window->GetTitle().ToString();
diff --git a/Engine/Source/Runtime/SlateRHIRenderer/Private/SlateRHIRenderingPolicy.cpp b/Engine/Source/Runtime/SlateRHIRenderer/Private/SlateRHIRenderingPolicy.cpp
--- a/Engine/Source/Runtime/SlateRHIRenderer/Private/SlateRHIRenderingPolicy.cpp
+++ b/Engine/Source/Runtime/SlateRHIRenderer/Private/SlateRHIRenderingPolicy.cpp
@@ -169,7 +169,7 @@ void FSlateRHIRenderingPolicy::BuildRenderingBuffers(FRHICommandListImmediate& R
 	SET_DWORD_STAT(STAT_SlateVertexCount, InBatchData.GetFinalVertexData().Num());
 }
 
-static FSceneView* CreateSceneView( FSceneViewFamilyContext* ViewFamilyContext, FSlateBackBuffer& BackBuffer, const FMatrix& ViewProjectionMatrix, const FIntRect InViewRect)
+static FSceneView* CreateSceneView( FSceneViewFamilyContext* ViewFamilyContext, FSlateBackBuffer& BackBuffer, const FMatrix& ViewProjectionMatrix, const FSlateRenderingParams& Params)
 {
 	QUICK_SCOPE_CYCLE_COUNTER(STAT_Slate_CreateSceneView);
 	// In loading screens, the engine is NULL, so we skip out.
@@ -178,7 +178,7 @@ static FSceneView* CreateSceneView( FSceneViewFamilyContext* ViewFamilyContext,
 		return nullptr;
 	}
 	
-	FIntRect ViewRect = InViewRect;
+	FIntRect ViewRect = Params.ViewRect;
 	if (ViewRect.IsEmpty())
 	{
 		ViewRect = FIntRect(FIntPoint(0, 0), BackBuffer.GetSizeXY());
@@ -212,6 +212,9 @@ static FSceneView* CreateSceneView( FSceneViewFamilyContext* ViewFamilyContext,
 		FViewMatrices()
 	);
 
+	// Always Update cursor position in realtime for slate
+	ViewUniformShaderParameters.CursorPosition = Params.CursorPostion;
+
 	// TODO LWC
 	ViewUniformShaderParameters.WorldViewOriginHigh = (FVector3f)View->ViewMatrices.GetViewOrigin();
 	
@@ -728,7 +731,7 @@ void FSlateRHIRenderingPolicy::DrawElements(
 				.SetTime(Params.Time)
 				.SetRealtimeUpdate(true)
 			);
-			SceneViews[i] = CreateSceneView(SceneViewFamilyContexts[i], BackBuffer, FMatrix(Params.ViewProjectionMatrix), Params.ViewRect);
+			SceneViews[i] = CreateSceneView(SceneViewFamilyContexts[i], BackBuffer, FMatrix(Params.ViewProjectionMatrix), Params);
 		}
 
 		SceneViewFamilyContexts[NumScenes - 1] = new FSceneViewFamilyContext
@@ -742,7 +745,7 @@ void FSlateRHIRenderingPolicy::DrawElements(
 			.SetTime(Params.Time)
 			.SetRealtimeUpdate(true)
 		);
-		SceneViews[NumScenes - 1] = CreateSceneView(SceneViewFamilyContexts[NumScenes - 1], BackBuffer, FMatrix(Params.ViewProjectionMatrix), Params.ViewRect);
+		SceneViews[NumScenes - 1] = CreateSceneView(SceneViewFamilyContexts[NumScenes - 1], BackBuffer, FMatrix(Params.ViewProjectionMatrix), Params);
 	}
 
 	TShaderMapRef<FSlateElementVS> GlobalVertexShader(GetGlobalShaderMap(GMaxRHIShaderPlatform));
diff --git a/Engine/Source/Runtime/SlateRHIRenderer/Private/SlateRHIRenderingPolicy.h b/Engine/Source/Runtime/SlateRHIRenderer/Private/SlateRHIRenderingPolicy.h
--- a/Engine/Source/Runtime/SlateRHIRenderer/Private/SlateRHIRenderingPolicy.h
+++ b/Engine/Source/Runtime/SlateRHIRenderer/Private/SlateRHIRenderingPolicy.h
@@ -34,6 +34,7 @@ struct FSlateRenderingParams
 	FMatrix44f ViewProjectionMatrix;
 	FVector2f ViewOffset;
 	FIntRect ViewRect;
+	FIntPoint CursorPostion;
 	FGameTime Time;
 	TRefCountPtr<IPooledRenderTarget> UITarget;
 	EDisplayColorGamut HDRDisplayColorGamut;
@@ -45,6 +46,7 @@ struct FSlateRenderingParams
 		: ViewProjectionMatrix(InViewProjectionMatrix)
 		, ViewOffset(0.f, 0.f)
 		, ViewRect(FIntRect())
+		, CursorPostion(FIntPoint())
 		, Time(InTime)
 		, HDRDisplayColorGamut(EDisplayColorGamut::sRGB_D65)
 		, UsedSlatePostBuffers(ESlatePostRT::None)
