Add support for multiple blur strengths within a single buffer in a single frame.

#tests PIE

diff --git a/Engine/Source/Runtime/SlateRHIRenderer/Private/FX/SlatePostBufferBlur.cpp b/Engine/Source/Runtime/SlateRHIRenderer/Private/FX/SlatePostBufferBlur.cpp
--- a/Engine/Source/Runtime/SlateRHIRenderer/Private/FX/SlatePostBufferBlur.cpp
+++ b/Engine/Source/Runtime/SlateRHIRenderer/Private/FX/SlatePostBufferBlur.cpp
@@ -4,6 +4,7 @@
 #include "UObject/StrongObjectPtr.h"
 #include "UObject/WeakObjectPtrTemplates.h"
 #include "RHIResources.h"
+#include "RHIUtilities.h"
 
 #include UE_INLINE_GENERATED_CPP_BY_NAME(SlatePostBufferBlur)
 
@@ -14,14 +15,25 @@ void FSlatePostBufferBlurProxy::PostProcess_Renderthread(FRHICommandListImmediat
 {
 	if (InRenderingPolicy.IsValid())
 	{
-		// Use rendering policy to perform blur post process with desired Src / Dst & respective extents
-		InRenderingPolicy.BlurRectExternal(RHICmdList, Src, Dst, SrcRect, DstRect, GaussianBlurStrength_RenderThread);
+		if (GaussianBlurStrength_RenderThread < UE_SMALL_NUMBER)
+		{
+			// No real blur, just copy
+			FRHICopyTextureInfo CopyInfo;
+			CopyInfo.SourcePosition = FIntVector(SrcRect.Min.X, SrcRect.Min.Y, 0);
+			CopyInfo.Size = FIntVector(DstRect.Width(), DstRect.Height(), 1);
+			TransitionAndCopyTexture(RHICmdList, Src, Dst, CopyInfo);
+		}
+		else
+		{
+			// Use rendering policy to perform blur post process with desired Src / Dst & respective extents
+			InRenderingPolicy.BlurRectExternal(RHICmdList, Src, Dst, SrcRect, DstRect, GaussianBlurStrength_RenderThread);
+		}
 	}
 }
 
 void FSlatePostBufferBlurProxy::OnUpdateValuesRenderThread()
 {
-	// Don't issue multiple updates in a single draw
+	// Don't issue multiple updates in a single frame from the CPU based on dirty values
 	if (!ParamUpdateFence.IsFenceComplete())
 	{
 		return;
@@ -32,21 +44,35 @@ void FSlatePostBufferBlurProxy::OnUpdateValuesRenderThread()
 	{
 		if (ParentBlurObject->GaussianBlurStrength != GaussianBlurStrength_RenderThread)
 		{
-			// Explicit param copy to avoid renderthread from reading value during gamethread write
-			float GaussianBlurStrengthCopy = ParentBlurObject->GaussianBlurStrength;
+			// Blur strengths can be updated from renderthread during draw or gamethread,
+			// if our parent object value matches the predraw then don't update the renderthread value.
+			// Instead we need to update our parent object's value to match the last value from renderthread
+			bool bUpdatedInRenderThread = ParentBlurObject->GaussianBlurStrength == GaussianBlurStrengthPreDraw;
 
-			// Execute param copy in a render command to safely update value on renderthread without race conditions
-			TWeakPtr<FSlatePostBufferBlurProxy> TempWeakThis = SharedThis(this);
-			ENQUEUE_RENDER_COMMAND(FUpdateValuesRenderThreadFX_Blur)([TempWeakThis, GaussianBlurStrengthCopy](FRHICommandListImmediate& RHICmdList)
+			if (bUpdatedInRenderThread)
 			{
-				if (TSharedPtr<FSlatePostBufferBlurProxy> SharedThisPin = TempWeakThis.Pin())
-				{
-					SharedThisPin->GaussianBlurStrength_RenderThread = GaussianBlurStrengthCopy;
-				}
-			});
+				ParentBlurObject->GaussianBlurStrength = GaussianBlurStrength_RenderThread;
+				GaussianBlurStrengthPreDraw = GaussianBlurStrength_RenderThread;
+			}
+			else
+			{
+				// Explicit param copy to avoid renderthread from reading value during gamethread write
+				float GaussianBlurStrengthCopy = ParentBlurObject->GaussianBlurStrength;
+				GaussianBlurStrengthPreDraw = GaussianBlurStrengthCopy;
 
-			// Issue fence to prevent multiple updates in a single draw
-			ParamUpdateFence.BeginFence();
+				// Execute param copy in a render command to safely update value on renderthread without race conditions
+				TWeakPtr<FSlatePostBufferBlurProxy> TempWeakThis = SharedThis(this);
+				ENQUEUE_RENDER_COMMAND(FUpdateValuesRenderThreadFX_Blur)([TempWeakThis, GaussianBlurStrengthCopy](FRHICommandListImmediate& RHICmdList)
+				{
+					if (TSharedPtr<FSlatePostBufferBlurProxy> SharedThisPin = TempWeakThis.Pin())
+					{
+						SharedThisPin->GaussianBlurStrength_RenderThread = GaussianBlurStrengthCopy;
+					}
+				});
+
+				// Issue fence to prevent multiple updates in a single frame
+				ParamUpdateFence.BeginFence();
+			}
 		}
 	}
 }
@@ -90,8 +116,19 @@ void USlatePostBufferBlur::PostProcess(FRenderResource* InViewInfo, FRenderResou
 			FIntRect SrcRect = FIntRect(0, 0, Src->GetSizeX(), Src->GetSizeY());
 			FIntRect DstRect = FIntRect(0, 0, DstExtent.X, DstExtent.Y);
 
-			// Use rendering policy to perform blur post process with desired Src / Dst & respective extents
-			InRenderingPolicy.BlurRectExternal(RHICmdList, Src, Dst, SrcRect, DstRect, GaussianBlurStrengthCopy);
+			if (GaussianBlurStrengthCopy < UE_SMALL_NUMBER)
+			{
+				// No real blur, just copy
+				FRHICopyTextureInfo CopyInfo;
+				CopyInfo.SourcePosition = FIntVector(SrcRect.Min.X, SrcRect.Min.Y, 0);
+				CopyInfo.Size = FIntVector(DstRect.Width(), DstRect.Height(), 1);
+				TransitionAndCopyTexture(RHICmdList, Src, Dst, CopyInfo);
+			}
+			else
+			{
+				// Use rendering policy to perform blur post process with desired Src / Dst & respective extents
+				InRenderingPolicy.BlurRectExternal(RHICmdList, Src, Dst, SrcRect, DstRect, GaussianBlurStrengthCopy);
+			}
 		}
 	});
 }
diff --git a/Engine/Source/Runtime/SlateRHIRenderer/Public/FX/SlatePostBufferBlur.h b/Engine/Source/Runtime/SlateRHIRenderer/Public/FX/SlatePostBufferBlur.h
--- a/Engine/Source/Runtime/SlateRHIRenderer/Public/FX/SlatePostBufferBlur.h
+++ b/Engine/Source/Runtime/SlateRHIRenderer/Public/FX/SlatePostBufferBlur.h
@@ -23,12 +23,19 @@ public:
 	virtual void OnUpdateValuesRenderThread() override;
 	//~ End FSlateRHIPostBufferProcessorProxy Interface
 
-protected:
-
 	/** Blur strength to use when processing, renderthread version actually used to draw. Must be updated via render command except during initialization. */
 	float GaussianBlurStrength_RenderThread = 10;
 
-	/** Fence to allow for us to queue only one update per draw command */
+	/** 
+	 * Blur strength can be updated from both renderthread during draw and gamethread update. 
+	 * Store the last value gamethread provided so we know if we should use the renderthread value or gamethread value. 
+	 * We will use the most recently updated one.
+	 */
+	float GaussianBlurStrengthPreDraw = 10;
+
+protected:
+
+	/** Fence to allow for us to queue only one update per draw command from the gamethread */
 	FRenderCommandFence ParamUpdateFence;
 };
 
diff --git a/Engine/Source/Runtime/UMG/Private/Components/PostBufferBlurUpdater.cpp b/Engine/Source/Runtime/UMG/Private/Components/PostBufferBlurUpdater.cpp
new file mode 100644
--- /dev/null
+++ b/Engine/Source/Runtime/UMG/Private/Components/PostBufferBlurUpdater.cpp
@@ -0,0 +1,45 @@
+// Copyright Epic Games, Inc. All Rights Reserved.
+
+#include "Components/PostBufferBlurUpdater.h"
+
+// Exclude SlateRHIRenderer related includes & implementations from server builds since the module is not a dependency / will not link for UMG on the server
+#if !UE_SERVER
+#include "FX/SlateRHIPostBufferProcessor.h"
+#include "FX/SlatePostBufferBlur.h"
+#endif // !UE_SERVER
+
+#include UE_INLINE_GENERATED_CPP_BY_NAME(PostBufferBlurUpdater)
+
+#define LOCTEXT_NAMESPACE "UMG"
+
+/////////////////////////////////////////////////////
+// UPostBufferBlurUpdater
+
+TSharedPtr<FSlatePostProcessorUpdaterProxy> UPostBufferBlurUpdater::GetRenderThreadProxy() const
+{
+	TSharedPtr<FSlatePostProcessorUpdaterProxy> Proxy = nullptr;
+#if !UE_SERVER
+	TSharedPtr<FPostBufferBlurUpdaterProxy> BlurProxy = MakeShared<FPostBufferBlurUpdaterProxy>();
+	BlurProxy->GaussianBlurStrength_RenderThread = GaussianBlurStrength;
+	Proxy = BlurProxy;
+#endif // !UE_SERVER
+	return Proxy;
+}
+
+
+/////////////////////////////////////////////////////
+// FPostBufferBlurUpdaterProxy
+
+void FPostBufferBlurUpdaterProxy::UpdateProcessor_RenderThread(TSharedPtr<FSlateRHIPostBufferProcessorProxy> InProcessor) const
+{
+#if !UE_SERVER
+	TSharedPtr<FSlatePostBufferBlurProxy> BlurRHIProxy = StaticCastSharedPtr<FSlatePostBufferBlurProxy>(InProcessor);
+	BlurRHIProxy->GaussianBlurStrength_RenderThread = GaussianBlurStrength_RenderThread;
+#endif // !UE_SERVER
+}
+
+
+/////////////////////////////////////////////////////
+
+#undef LOCTEXT_NAMESPACE
+
diff --git a/Engine/Source/Runtime/UMG/Private/Components/PostBufferUpdate.cpp b/Engine/Source/Runtime/UMG/Private/Components/PostBufferUpdate.cpp
--- a/Engine/Source/Runtime/UMG/Private/Components/PostBufferUpdate.cpp
+++ b/Engine/Source/Runtime/UMG/Private/Components/PostBufferUpdate.cpp
@@ -14,7 +14,10 @@
 UPostBufferUpdate::UPostBufferUpdate()
 {
 	bPerformDefaultPostBufferUpdate = true;
+PRAGMA_DISABLE_DEPRECATION_WARNINGS
 	BuffersToUpdate = {};
+PRAGMA_ENABLE_DEPRECATION_WARNINGS
+	UpdateBufferInfos = {};
 }
 
 void UPostBufferUpdate::SetPerformDefaultPostBufferUpdate(bool bInPerformDefaultPostBufferUpdate)
@@ -45,7 +48,31 @@ TSharedRef<SWidget> UPostBufferUpdate::RebuildWidget()
 
 	if (bSetBuffersToUpdate)
 	{
-		MyPostBufferUpdate->SetBuffersToUpdate(BuffersToUpdate);
+PRAGMA_DISABLE_DEPRECATION_WARNINGS
+		if (!UpdateBufferInfos.IsEmpty())
+		{
+			TArray<ESlatePostRT> TempBuffersToUpdate;
+			TMap<ESlatePostRT, TSharedPtr<FSlatePostProcessorUpdaterProxy>> TempProcessorUpdaters;
+
+			for (const FSlatePostBufferUpdateInfo& UpdateBufferInfo : UpdateBufferInfos)
+			{
+				TempBuffersToUpdate.Add(UpdateBufferInfo.BufferToUpdate);
+				if (TObjectPtr<USlatePostBufferProcessorUpdater> PostParamUpdater = UpdateBufferInfo.PostParamUpdater)
+				{
+					TSharedPtr<FSlatePostProcessorUpdaterProxy> PostParamUpdaterProxy = PostParamUpdater->GetRenderThreadProxy();
+					PostParamUpdaterProxy->bSkipBufferUpdate = PostParamUpdater->bSkipBufferUpdate;
+					TempProcessorUpdaters.Add(UpdateBufferInfo.BufferToUpdate, PostParamUpdaterProxy);
+				}
+			}
+
+			MyPostBufferUpdate->SetBuffersToUpdate(TempBuffersToUpdate);
+			MyPostBufferUpdate->SetProcessorUpdaters(TempProcessorUpdaters);
+		}
+		else
+		{
+			MyPostBufferUpdate->SetBuffersToUpdate(BuffersToUpdate);
+		}
+PRAGMA_ENABLE_DEPRECATION_WARNINGS
 	}
 
 	return MyPostBufferUpdate.ToSharedRef();
@@ -73,7 +100,31 @@ void UPostBufferUpdate::SynchronizeProperties()
 
 	if (bSetBuffersToUpdate)
 	{
-		MyPostBufferUpdate->SetBuffersToUpdate(BuffersToUpdate);
+PRAGMA_DISABLE_DEPRECATION_WARNINGS
+		if (!UpdateBufferInfos.IsEmpty())
+		{
+			TArray<ESlatePostRT> TempBuffersToUpdate;
+			TMap<ESlatePostRT, TSharedPtr<FSlatePostProcessorUpdaterProxy>> TempProcessorUpdaters;
+
+			for (const FSlatePostBufferUpdateInfo& UpdateBufferInfo : UpdateBufferInfos)
+			{
+				TempBuffersToUpdate.Add(UpdateBufferInfo.BufferToUpdate);
+				if (TObjectPtr<USlatePostBufferProcessorUpdater> PostParamUpdater = UpdateBufferInfo.PostParamUpdater)
+				{
+					TSharedPtr<FSlatePostProcessorUpdaterProxy> PostParamUpdaterProxy = PostParamUpdater->GetRenderThreadProxy();
+					PostParamUpdaterProxy->bSkipBufferUpdate = PostParamUpdater->bSkipBufferUpdate;
+					TempProcessorUpdaters.Add(UpdateBufferInfo.BufferToUpdate, PostParamUpdaterProxy);
+				}
+			}
+
+			MyPostBufferUpdate->SetBuffersToUpdate(TempBuffersToUpdate);
+			MyPostBufferUpdate->SetProcessorUpdaters(TempProcessorUpdaters);
+		}
+		else
+		{
+			MyPostBufferUpdate->SetBuffersToUpdate(BuffersToUpdate);
+		}
+PRAGMA_ENABLE_DEPRECATION_WARNINGS
 	}
 }
 
diff --git a/Engine/Source/Runtime/UMG/Private/Slate/SPostBufferUpdate.cpp b/Engine/Source/Runtime/UMG/Private/Slate/SPostBufferUpdate.cpp
--- a/Engine/Source/Runtime/UMG/Private/Slate/SPostBufferUpdate.cpp
+++ b/Engine/Source/Runtime/UMG/Private/Slate/SPostBufferUpdate.cpp
@@ -49,6 +49,9 @@ public:
 	 * Additionally, note that this value should be masked against the currently enabled buffers in 'USlateRHIRendererSettings'
 	 */
 	ESlatePostRT BuffersToUpdate_Renderthread = ESlatePostRT::None;
+
+	/** Proxies used to update a post processor within a frame */
+	TMap<ESlatePostRT, TSharedPtr<FSlatePostProcessorUpdaterProxy>> ProcessorUpdaters;
 };
 
 /////////////////////////////////////////////////////
@@ -98,6 +101,19 @@ void FPostBufferUpdater::Draw_RHIRenderThread(FRHICommandListImmediate& RHICmdLi
 
 				if (TSharedPtr<FSlateRHIPostBufferProcessorProxy> PostProcessorProxy = USlateFXSubsystem::GetPostProcessorProxy(SlatePostBufferBit))
 				{
+					if (TSharedPtr<FSlatePostProcessorUpdaterProxy>* ProcessorUpdaterItr = ProcessorUpdaters.Find(SlatePostBufferBit))
+					{
+						if (TSharedPtr<FSlatePostProcessorUpdaterProxy> ProcessorUpdater = *ProcessorUpdaterItr)
+						{
+							ProcessorUpdater->UpdateProcessor_RenderThread(PostProcessorProxy);
+
+							if (ProcessorUpdater->bSkipBufferUpdate)
+							{
+								return;
+							}
+						}
+					}
+
 					PostProcessorProxy->PostProcess_Renderthread(RHICmdList, Src, Dst, SrcRect, DstRect, RenderingPolicyInterface);
 				}
 				else
@@ -218,6 +234,16 @@ void SPostBufferUpdate::SetBuffersToUpdate(const TArrayView<ESlatePostRT> InBuff
 #endif // !UE_SERVER
 }
 
+void SPostBufferUpdate::SetProcessorUpdaters(TMap<ESlatePostRT, TSharedPtr<FSlatePostProcessorUpdaterProxy>> InProcessorUpdaters)
+{
+#if !UE_SERVER
+	if (PostBufferUpdater)
+	{
+		PostBufferUpdater->ProcessorUpdaters = InProcessorUpdaters;
+	}
+#endif // !UE_SERVER
+}
+
 const TArrayView<const ESlatePostRT> SPostBufferUpdate::GetBuffersToUpdate() const
 {
 	return MakeArrayView(BuffersToUpdate);
@@ -227,7 +253,7 @@ UMG_API void SPostBufferUpdate::ReleasePostBufferUpdater()
 {
 #if !UE_SERVER
 	// Copy the pointer onto a lambda to defer the final deletion to after any pending uses on the renderthread
-	TSharedPtr<FPostBufferUpdater, ESPMode::ThreadSafe> ReleaseMe = PostBufferUpdater;
+	TSharedPtr<FPostBufferUpdater> ReleaseMe = PostBufferUpdater;
 	ENQUEUE_RENDER_COMMAND(ReleaseCommand)([ReleaseMe](FRHICommandList& RHICmdList) mutable
 	{
 		ReleaseMe.Reset();
diff --git a/Engine/Source/Runtime/UMG/Public/Components/PostBufferBlurUpdater.h b/Engine/Source/Runtime/UMG/Public/Components/PostBufferBlurUpdater.h
new file mode 100644
--- /dev/null
+++ b/Engine/Source/Runtime/UMG/Public/Components/PostBufferBlurUpdater.h
@@ -0,0 +1,43 @@
+// Copyright Epic Games, Inc. All Rights Reserved.
+
+#pragma once
+
+#include "CoreMinimal.h"
+#include "UObject/ObjectMacros.h"
+#include "Misc/Attribute.h"
+#include "Widgets/SWidget.h"
+#include "Components/PostBufferUpdate.h"
+#include "Rendering/SlateRendererTypes.h"
+#include "Slate/SPostBufferUpdate.h"
+
+#include "PostBufferBlurUpdater.generated.h"
+
+/**
+ * Processor updater that sets the blur strength for a blur processor intra-frame on the renderthread
+ */
+UCLASS(NotBlueprintable, MinimalAPI)
+class UPostBufferBlurUpdater : public USlatePostBufferProcessorUpdater
+{
+	GENERATED_BODY()
+
+public:
+
+	UPROPERTY(EditAnywhere, Category = "GaussianBlur")
+	float GaussianBlurStrength = 10;
+
+public:
+	virtual TSharedPtr<FSlatePostProcessorUpdaterProxy > GetRenderThreadProxy() const override;
+};
+
+/**
+ * Renderthread proxy for the blur processor updater
+ */
+class FPostBufferBlurUpdaterProxy : public FSlatePostProcessorUpdaterProxy 
+{
+public:
+
+	/** Blur strength that will be copied over to the processor mid-frame */
+	float GaussianBlurStrength_RenderThread = 10;
+
+	virtual void UpdateProcessor_RenderThread(TSharedPtr<FSlateRHIPostBufferProcessorProxy>) const override;
+};
\ No newline at end of file
diff --git a/Engine/Source/Runtime/UMG/Public/Components/PostBufferUpdate.h b/Engine/Source/Runtime/UMG/Public/Components/PostBufferUpdate.h
--- a/Engine/Source/Runtime/UMG/Public/Components/PostBufferUpdate.h
+++ b/Engine/Source/Runtime/UMG/Public/Components/PostBufferUpdate.h
@@ -12,6 +12,26 @@
 #include "PostBufferUpdate.generated.h"
 
 class SPostBufferUpdate;
+class USlatePostBufferProcessorUpdater;
+class FSlatePostProcessorUpdaterProxy ;
+class FSlateRHIPostBufferProcessorProxy;
+
+/**
+ * Struct containing info needed to update a particular buffer
+ */
+USTRUCT()
+struct UMG_API FSlatePostBufferUpdateInfo
+{
+	GENERATED_BODY()
+
+	/** Buffers that we should update, all of these buffers will be affected by 'bPerformDefaultPostBufferUpdate' if disabled */
+	UPROPERTY(EditAnywhere, Category = "Behavior", meta = (AllowPrivateAccess = "true"))
+	ESlatePostRT BufferToUpdate = ESlatePostRT::None;
+
+	/** Optional processor updater for buffer, used to update a processor within a frame */
+	UPROPERTY(EditAnywhere, Instanced, Category = "Behavior", meta = (AllowPrivateAccess = "true"))
+	TObjectPtr<USlatePostBufferProcessorUpdater> PostParamUpdater = nullptr;
+};
 
 /**
  * Widget that when drawn, will trigger the slate post buffer to update. Does not draw anything itself.
@@ -36,10 +56,15 @@ private:
 	UPROPERTY(EditAnywhere, BlueprintReadWrite, Setter="SetPerformDefaultPostBufferUpdate", Category = "Behavior", meta = (AllowPrivateAccess = "true"))
 	bool bPerformDefaultPostBufferUpdate;
 
+	UE_DEPRECATED(5.5, "BuffersToUpdate is deprecated. Please use UpdateBufferInfos. This array will be ignored if UpdateBufferInfos is used")
 	/** Buffers that we should update, all of these buffers will be affected by 'bPerformDefaultPostBufferUpdate' if disabled */
 	UPROPERTY(EditAnywhere, Category = "Behavior", meta = (AllowPrivateAccess = "true"))
 	TArray<ESlatePostRT> BuffersToUpdate;
 
+	/** Buffer to update when this widget is drawn, along with info needed to update that buffer if desired intra-frame */
+	UPROPERTY(EditAnywhere, Category = "Behavior", meta = (AllowPrivateAccess = "true"))
+	TArray<FSlatePostBufferUpdateInfo> UpdateBufferInfos;
+
 public:
 	/** Set the orientation of the stack box. The existing elements will be rearranged. */
 	UMG_API void SetPerformDefaultPostBufferUpdate(bool bInPerformDefaultPostBufferUpdate);
@@ -54,3 +79,29 @@ protected:
 protected:
 	TSharedPtr<SPostBufferUpdate> MyPostBufferUpdate;
 };
+
+/**
+ * Class that can create a FPostParamUpdaterProxy whose lifetime
+ * will be managed by the renderthread. This proxy will be given a 
+ * Post buffer processor to update mid-frame.
+ */
+UCLASS(MinimalAPI, Abstract, Blueprintable, EditInlineNew, CollapseCategories)
+class USlatePostBufferProcessorUpdater : public UObject
+{
+	GENERATED_BODY()
+
+public:
+
+	/**
+	 * True implies we will skip the buffer update & only update the processor.
+	 * Useful to reset params for processor runs next frame
+	 */
+	UPROPERTY(EditAnywhere, Category = "Update")
+	bool bSkipBufferUpdate = false;
+
+public:
+	virtual TSharedPtr<FSlatePostProcessorUpdaterProxy> GetRenderThreadProxy() const
+	{
+		return nullptr;
+	};
+};
\ No newline at end of file
diff --git a/Engine/Source/Runtime/UMG/Public/Slate/SPostBufferUpdate.h b/Engine/Source/Runtime/UMG/Public/Slate/SPostBufferUpdate.h
--- a/Engine/Source/Runtime/UMG/Public/Slate/SPostBufferUpdate.h
+++ b/Engine/Source/Runtime/UMG/Public/Slate/SPostBufferUpdate.h
@@ -11,6 +11,7 @@
 
 class FPaintArgs;
 class FSlateWindowElementList;
+class FSlateRHIPostBufferProcessorProxy;
 
 /**
  * Custom Slate drawer to update slate post buffer
@@ -19,6 +20,30 @@ class FSlateWindowElementList;
  */
 class FPostBufferUpdater;
 
+/**
+ * Class that can update a given Slate post buffer processor via its renderthread proxy.
+ *
+ * This proxy is also deleted on the renderthread
+ */
+class FSlatePostProcessorUpdaterProxy : public TSharedFromThis<FSlatePostProcessorUpdaterProxy>
+{
+public:
+
+	/**
+	 * True implies we will skip the buffer update & only update the processor.
+	 * Useful to reset params for processor runs next frame
+	 */
+	bool bSkipBufferUpdate = false;
+
+public:
+	
+	virtual ~FSlatePostProcessorUpdaterProxy ()
+	{
+	}
+
+	virtual void UpdateProcessor_RenderThread(TSharedPtr<FSlateRHIPostBufferProcessorProxy>) const = 0;
+};
+
 /**
  * Implements a widget that triggers a post buffer update on draw
  */
@@ -62,6 +87,11 @@ public:
 	 */
 	UMG_API void SetBuffersToUpdate(const TArrayView<ESlatePostRT> InBuffersToUpdate);
 
+	/**
+	 * Set processsor updaters for the given buffers by index.
+	 */
+	UMG_API void SetProcessorUpdaters(TMap<ESlatePostRT, TSharedPtr<FSlatePostProcessorUpdaterProxy>> InProcessorUpdaters);
+
 	/** Get buffers to update */
 	UMG_API const TArrayView<const ESlatePostRT> GetBuffersToUpdate() const;
 
@@ -89,5 +119,5 @@ private:
 	TArray<ESlatePostRT> BuffersToUpdate;
 
 	/** Custom drawer used to trigger a post buffer update */
-	TSharedPtr<FPostBufferUpdater, ESPMode::ThreadSafe> PostBufferUpdater;
+	TSharedPtr<FPostBufferUpdater> PostBufferUpdater;
 };
