Expose ViewportScaleUI to view uniform

#tests PIE, material & UMG editor,

diff --git a/Engine/Source/Editor/MaterialEditor/Private/SMaterialEditorViewport.cpp b/Engine/Source/Editor/MaterialEditor/Private/SMaterialEditorViewport.cpp
--- a/Engine/Source/Editor/MaterialEditor/Private/SMaterialEditorViewport.cpp
+++ b/Engine/Source/Editor/MaterialEditor/Private/SMaterialEditorViewport.cpp
@@ -998,7 +998,13 @@ FChildren* SMaterialEditorUIPreviewZoomer::GetChildren()
 int32 SMaterialEditorUIPreviewZoomer::OnPaint(const FPaintArgs& Args, const FGeometry& AllottedGeometry, const FSlateRect& MyCullingRect, FSlateWindowElementList& OutDrawElements, int32 LayerId, const FWidgetStyle& InWidgetStyle, bool bParentEnabled) const
 {
 	LayerId = SPanel::OnPaint(Args, AllottedGeometry, MyCullingRect, OutDrawElements, LayerId, InWidgetStyle, bParentEnabled);
-	
+
+	// Set a UI scale for materials to use as reference, done on a per-window basis since don't want to change global uniforms per element
+	if (SWindow* ParentWindow = OutDrawElements.GetPaintWindow())
+	{
+		ParentWindow->SetViewportScaleUIOverride(ZoomLevel);
+	}
+
 	if (IsCurrentlyScrollable())
 	{
 		LayerId = ScrollyZoomy.PaintSoftwareCursorIfNeeded(AllottedGeometry, MyCullingRect, OutDrawElements, LayerId);
diff --git a/Engine/Source/Editor/UMGEditor/Private/Designer/SDesignerView.cpp b/Engine/Source/Editor/UMGEditor/Private/Designer/SDesignerView.cpp
--- a/Engine/Source/Editor/UMGEditor/Private/Designer/SDesignerView.cpp
+++ b/Engine/Source/Editor/UMGEditor/Private/Designer/SDesignerView.cpp
@@ -2422,6 +2422,12 @@ void SDesignerView::OnPaintBackground(const FGeometry& AllottedGeometry, const F
 		}
 	}
 
+	// Set a UI scale for materials to use as reference, done on a per-window basis since don't want to change global uniforms per element
+	if (SWindow* ParentWindow = OutDrawElements.GetPaintWindow())
+	{
+		ParentWindow->SetViewportScaleUIOverride(GetZoomAmount());
+	}
+
 	SDesignSurface::OnPaintBackground(AllottedGeometry, MyCullingRect, OutDrawElements, LayerId);
 
 	if (bShowResolutionOutlines)
diff --git a/Engine/Source/Runtime/Engine/Public/SceneView.h b/Engine/Source/Runtime/Engine/Public/SceneView.h
--- a/Engine/Source/Runtime/Engine/Public/SceneView.h
+++ b/Engine/Source/Runtime/Engine/Public/SceneView.h
@@ -934,6 +934,7 @@ enum ETranslucencyVolumeCascade
 	VIEW_UNIFORM_BUFFER_MEMBER(float, NotCoveredMinStepScale) \
 	VIEW_UNIFORM_BUFFER_MEMBER(float, DitheredTransparencyStepThreshold) \
 	VIEW_UNIFORM_BUFFER_MEMBER(float, DitheredTransparencyTraceThreshold) \
+	VIEW_UNIFORM_BUFFER_MEMBER(float, ViewportScaleUI) \
 	VIEW_UNIFORM_BUFFER_MEMBER(FIntPoint, CursorPosition) \
 	VIEW_UNIFORM_BUFFER_MEMBER(float, bCheckerboardSubsurfaceProfileRendering) \
 	VIEW_UNIFORM_BUFFER_MEMBER(FVector3f, VolumetricFogInvGridSize) \
diff --git a/Engine/Source/Runtime/SlateCore/Private/Widgets/SWindow.cpp b/Engine/Source/Runtime/SlateCore/Private/Widgets/SWindow.cpp
--- a/Engine/Source/Runtime/SlateCore/Private/Widgets/SWindow.cpp
+++ b/Engine/Source/Runtime/SlateCore/Private/Widgets/SWindow.cpp
@@ -2029,6 +2029,7 @@ SWindow::SWindow()
 	, PreFullscreenPosition( FVector2f::ZeroVector )
 	, Size( FVector2f::ZeroVector )
 	, ViewportSize( FVector2f::ZeroVector )
+	, ViewportScaleUIOverride( -1.f )
 	, TitleBarSize( SWindowDefs::DefaultTitleBarSize )
 	, ContentSlot(nullptr)
 	, Style( &FCoreStyle::Get().GetWidgetStyle<FWindowStyle>("Window") )
diff --git a/Engine/Source/Runtime/SlateCore/Public/Widgets/SWindow.h b/Engine/Source/Runtime/SlateCore/Public/Widgets/SWindow.h
--- a/Engine/Source/Runtime/SlateCore/Public/Widgets/SWindow.h
+++ b/Engine/Source/Runtime/SlateCore/Public/Widgets/SWindow.h
@@ -881,6 +881,18 @@ public:
 
 	bool GetIsHDR() const { return bIsHDR; }
 
+	void SetViewportScaleUIOverride(float InViewportScaleUIOverride)
+	{
+		ViewportScaleUIOverride = InViewportScaleUIOverride;
+	}
+
+	void ResetViewportScaleUIOverride()
+	{
+		ViewportScaleUIOverride = -1.0f;
+	}
+
+	float GetViewportScaleUIOverride() const { return ViewportScaleUIOverride; }
+
 	bool IsVirtualWindow() const { return bVirtualWindow; }
 
 	bool IsMirrorWindow()
@@ -1196,6 +1208,13 @@ protected:
 	/** Pointer to the viewport registered with this window if any */
 	TWeakPtr<ISlateViewport> Viewport;
 
+	/** 
+	 * Mix of aspect ratio + DPI + UI Zoom. Used in shaders / materials.
+	 * Can be overriden by child widgets, is reset after draw every frame.
+	 * Negative value indicates not used, in which case application level UI scale will be used.
+	 */
+	float ViewportScaleUIOverride;
+
 	/** Size of this window's title bar.  Can be zero.  Set at construction and should not be changed afterwards. */
 	float TitleBarSize;
 
diff --git a/Engine/Source/Runtime/SlateRHIRenderer/Private/SlateRHIRenderer.cpp b/Engine/Source/Runtime/SlateRHIRenderer/Private/SlateRHIRenderer.cpp
--- a/Engine/Source/Runtime/SlateRHIRenderer/Private/SlateRHIRenderer.cpp
+++ b/Engine/Source/Runtime/SlateRHIRenderer/Private/SlateRHIRenderer.cpp
@@ -12,6 +12,7 @@
 #include "EngineGlobals.h"
 #include "Engine/AssetManager.h"
 #include "Engine/TextureRenderTarget2D.h"
+#include "Engine/UserInterfaceSettings.h"
 #include "FX/SlateFXSubsystem.h"
 #include "FX/SlateRHIPostBufferProcessor.h"
 #include "Materials/MaterialRenderProxy.h"
@@ -142,6 +143,7 @@ struct FSlateDrawWindowCommandParams
 	SWindow* Window;
 	FIntPoint CursorPostion;
 	FIntRect ViewRect;
+	float ViewportScaleUI;
 	ESlatePostRT UsedSlatePostBuffers;
 #if WANTS_DRAW_MESH_EVENTS
 	FString WindowTitle;
@@ -1023,6 +1025,7 @@ void RenderSlateBatch(FTextureRHIRef SlateRenderTarget, bool bClear, bool bIsHDR
 				RenderParams.HDRDisplayColorGamut = ViewportInfo.HDRDisplayColorGamut;
 				RenderParams.ViewRect = DrawCommandParams.ViewRect;
 				RenderParams.CursorPostion = DrawCommandParams.CursorPostion;
+				RenderParams.ViewportScaleUI = DrawCommandParams.ViewportScaleUI;
 				RenderParams.UsedSlatePostBuffers = DrawCommandParams.UsedSlatePostBuffers;
 				if (ViewportInfo.bSceneHDREnabled && !bIsHDR)
 				{
@@ -1941,6 +1944,7 @@ void FSlateRHIRenderer::DrawWindows_Private(FSlateDrawBuffer& WindowDrawBuffer)
 					Params.WindowElementList = &ElementList;
 					Params.Window = Window;
 					Params.ViewRect = GetViewRect();
+					Params.ViewportScaleUI = Window->GetViewportScaleUIOverride();
 					Params.CursorPostion = FIntPoint(CursorPosition.X, CursorPosition.Y);
 					Params.UsedSlatePostBuffers = UsedSlatePostBufferBits;
 #if WANTS_DRAW_MESH_EVENTS
@@ -1956,6 +1960,12 @@ void FSlateRHIRenderer::DrawWindows_Private(FSlateDrawBuffer& WindowDrawBuffer)
 						FPlatformTime::Seconds() - GStartTime, (float)FApp::GetDeltaTime(),
 						FApp::GetCurrentTime() - GStartTime, (float)FApp::GetDeltaTime());
 
+					// Negative Viewport UI Scale, get the value from user settings
+					if (Params.ViewportScaleUI < 0)
+					{	
+						Params.ViewportScaleUI = GetDefault<UUserInterfaceSettings>()->GetDPIScaleBasedOnSize(Params.ViewRect.Size());
+					}
+
 					// Skip the actual draw if we're in a headless execution environment
 					bool bLocalTakingAScreenShot = bTakingAScreenShot;
 					if (GIsClient && !IsRunningCommandlet() && !GUsingNullRHI)
@@ -1988,6 +1998,9 @@ void FSlateRHIRenderer::DrawWindows_Private(FSlateDrawBuffer& WindowDrawBuffer)
 
 					SlateWindowRendered.Broadcast(*Window, &ViewInfo->ViewportRHI);
 
+					// Now that all widgets have drawn, reset material UI scale
+					Window->ResetViewportScaleUIOverride();
+
 					if (bLocalTakingAScreenShot)
 					{
 						// immediately flush the RHI command list
diff --git a/Engine/Source/Runtime/SlateRHIRenderer/Private/SlateRHIRenderingPolicy.cpp b/Engine/Source/Runtime/SlateRHIRenderer/Private/SlateRHIRenderingPolicy.cpp
--- a/Engine/Source/Runtime/SlateRHIRenderer/Private/SlateRHIRenderingPolicy.cpp
+++ b/Engine/Source/Runtime/SlateRHIRenderer/Private/SlateRHIRenderingPolicy.cpp
@@ -212,6 +212,9 @@ static FSceneView* CreateSceneView( FSceneViewFamilyContext* ViewFamilyContext,
 		FViewMatrices()
 	);
 
+	// Update Viewport Scale UI from any external sources (Material editor, UMG zoom scale / etc).
+	ViewUniformShaderParameters.ViewportScaleUI = Params.ViewportScaleUI;
+
 	// Always Update cursor position in realtime for slate
 	ViewUniformShaderParameters.CursorPosition = Params.CursorPostion;
 
diff --git a/Engine/Source/Runtime/SlateRHIRenderer/Private/SlateRHIRenderingPolicy.h b/Engine/Source/Runtime/SlateRHIRenderer/Private/SlateRHIRenderingPolicy.h
--- a/Engine/Source/Runtime/SlateRHIRenderer/Private/SlateRHIRenderingPolicy.h
+++ b/Engine/Source/Runtime/SlateRHIRenderer/Private/SlateRHIRenderingPolicy.h
@@ -38,6 +38,7 @@ struct FSlateRenderingParams
 	FGameTime Time;
 	TRefCountPtr<IPooledRenderTarget> UITarget;
 	EDisplayColorGamut HDRDisplayColorGamut;
+	float ViewportScaleUI;
 	ESlatePostRT UsedSlatePostBuffers;
 	bool bWireFrame;
 	bool bIsHDR;
@@ -49,6 +50,7 @@ struct FSlateRenderingParams
 		, CursorPostion(FIntPoint())
 		, Time(InTime)
 		, HDRDisplayColorGamut(EDisplayColorGamut::sRGB_D65)
+		, ViewportScaleUI(1.0f)
 		, UsedSlatePostBuffers(ESlatePostRT::None)
 		, bWireFrame(false)
 		, bIsHDR(false)
