Temporarily only set ViewportUI Scale in editor, since in non-editor acquring scale can permanently leave DPI with an invalid value

#tests PIE, Standalone

diff --git a/Engine/Source/Runtime/SlateRHIRenderer/Private/SlateRHIRenderer.cpp b/Engine/Source/Runtime/SlateRHIRenderer/Private/SlateRHIRenderer.cpp
--- a/Engine/Source/Runtime/SlateRHIRenderer/Private/SlateRHIRenderer.cpp
+++ b/Engine/Source/Runtime/SlateRHIRenderer/Private/SlateRHIRenderer.cpp
@@ -1964,7 +1964,7 @@ void FSlateRHIRenderer::DrawWindows_Private(FSlateDrawBuffer& WindowDrawBuffer)
 						FApp::GetCurrentTime() - GStartTime, (float)FApp::GetDeltaTime());
 
 					// Negative Viewport UI Scale, get the value from user settings
-					if (Params.ViewportScaleUI < 0)
+					if (GIsEditor && Params.ViewportScaleUI < 0)
 					{	
 						Params.ViewportScaleUI = GetDefault<UUserInterfaceSettings>()->GetDPIScaleBasedOnSize(Params.ViewRect.Size());
 					}
