Skip postbuffer initialization in separate thread materials such as movie players.

#tests PIE, Standalone, Repro project fixed

diff --git a/Engine/Source/Runtime/SlateRHIRenderer/Private/SlateMaterialResource.cpp b/Engine/Source/Runtime/SlateRHIRenderer/Private/SlateMaterialResource.cpp
--- a/Engine/Source/Runtime/SlateRHIRenderer/Private/SlateMaterialResource.cpp
+++ b/Engine/Source/Runtime/SlateRHIRenderer/Private/SlateMaterialResource.cpp
@@ -90,7 +90,7 @@ FSlateMaterialResource::FSlateMaterialResource(const UMaterialInterface& InMater
 	SlateProxy->ActualSize = InImageSize.IntPoint();
 	SlateProxy->Resource = this;
 
-	if (MaterialObject)
+	if (MaterialObject && (IsInRenderingThread() || IsInGameThread()))
 	{
 		// Quality / Feature level irrelevant since flag to search all levels for both is true
 		TArray<UTexture*> OutUsedTextures;
@@ -156,7 +156,7 @@ void FSlateMaterialResource::UpdateMaterial(const UMaterialInterface& InMaterial
 	SlateMaterialResource::CheckInvalidMaterialProxy(MaterialProxy, DebugName);
 #endif
 
-	if (MaterialObject)
+	if (MaterialObject && (IsInRenderingThread() || IsInGameThread()))
 	{
 		// Quality / Feature level irrelevant since flag to search all levels for both is true
 		TArray<UTexture*> OutUsedTextures;
