Fix Slate PostBuffer not working on formats other than PF_A2B10G10R10.

#tests PIE, Standalone, all 5 ini formats

diff --git a/Engine/Source/Runtime/Engine/Private/Slate/SceneViewport.cpp b/Engine/Source/Runtime/Engine/Private/Slate/SceneViewport.cpp
--- a/Engine/Source/Runtime/Engine/Private/Slate/SceneViewport.cpp
+++ b/Engine/Source/Runtime/Engine/Private/Slate/SceneViewport.cpp
@@ -57,6 +57,7 @@ FSceneViewport::FSceneViewport( FViewportClient* InViewportClient, TSharedPtr<SV
 	, bHDRViewport(false)
 	, MousePosBeforeHiddenDueToCapture( -1, -1 )
 	, RTTSize( 0, 0 )
+	, SceneTargetFormat( EPixelFormat::PF_A2B10G10R10 )
 	, CurrentBufferedTargetIndex(0)
 	, NextBufferedTargetIndex(0)
 	, NumTouches(0)
@@ -2070,20 +2071,21 @@ void FSceneViewport::InitRHI(FRHICommandListBase& RHICmdList)
 	RTTSize = FIntPoint(0, 0);
 
 	uint32 TexSizeX = SizeX, TexSizeY = SizeY;
+
+	static const auto CVarDefaultBackBufferPixelFormat = IConsoleManager::Get().FindTConsoleVariableDataInt(TEXT("r.DefaultBackBufferPixelFormat"));
+	SceneTargetFormat = EDefaultBackBufferPixelFormat::Convert2PixelFormat(EDefaultBackBufferPixelFormat::FromInt(CVarDefaultBackBufferPixelFormat->GetValueOnRenderThread()));
+	SceneTargetFormat = RHIPreferredPixelFormatHint(SceneTargetFormat);
+
+	if (bHDRViewport)
+	{
+		SceneTargetFormat = GRHIHDRDisplayOutputFormat;
+	}
+
 	if (UseSeparateRenderTarget())
 	{
 		int32 NumBufferedFrames = 1;
 		TArray<FTextureRHIRef> BufferedRTRHI;
 		TArray<FTextureRHIRef> BufferedSRVRHI;
-
-		static const auto CVarDefaultBackBufferPixelFormat = IConsoleManager::Get().FindTConsoleVariableDataInt(TEXT("r.DefaultBackBufferPixelFormat"));
-		EPixelFormat SceneTargetFormat = EDefaultBackBufferPixelFormat::Convert2PixelFormat(EDefaultBackBufferPixelFormat::FromInt(CVarDefaultBackBufferPixelFormat->GetValueOnRenderThread()));
-		SceneTargetFormat = RHIPreferredPixelFormatHint(SceneTargetFormat);
-
-		if (bHDRViewport)
-		{
-			SceneTargetFormat = GRHIHDRDisplayOutputFormat;
-		}
 		
 		// @todo vreditor switch: This code needs to be called when switching between stereo/non when going immersive.  Seems to always work out that way anyway though? (Probably due to resize)
 		bool bHMDAllocatedSeparateRenderTargets = false;
diff --git a/Engine/Source/Runtime/Engine/Public/Slate/SceneViewport.h b/Engine/Source/Runtime/Engine/Public/Slate/SceneViewport.h
--- a/Engine/Source/Runtime/Engine/Public/Slate/SceneViewport.h
+++ b/Engine/Source/Runtime/Engine/Public/Slate/SceneViewport.h
@@ -263,6 +263,9 @@ public:
 
 	ENGINE_API virtual FSlateShaderResource* GetViewportRenderTargetTexture() override;
 
+	/** Returns format for the scene of this viewport. */
+	ENGINE_API EPixelFormat GetSceneTargetFormat() const override { return SceneTargetFormat; }
+
 	/** Get the cached viewport geometry. */
 	const FGeometry& GetCachedGeometry() const { return CachedGeometry; }
 
@@ -426,6 +429,8 @@ private:
 	FIntPoint MousePosBeforeHiddenDueToCapture;
 	/** Dimensions of RenderTarget texture. */
 	FIntPoint RTTSize;
+	/** Pixel format of all Buffered RenderTarget textures. */
+	EPixelFormat SceneTargetFormat;
 
 	/** Reprojection on some HMD RHI's requires ViewportTargets to be buffered */
 	/** The render target used by Slate to draw the viewport.  Can be null if this viewport renders directly to the backbuffer */
diff --git a/Engine/Source/Runtime/SlateCore/Public/Rendering/RenderingCommon.h b/Engine/Source/Runtime/SlateCore/Public/Rendering/RenderingCommon.h
--- a/Engine/Source/Runtime/SlateCore/Public/Rendering/RenderingCommon.h
+++ b/Engine/Source/Runtime/SlateCore/Public/Rendering/RenderingCommon.h
@@ -10,6 +10,7 @@
 #include "Input/Reply.h"
 #include "Input/NavigationReply.h"
 #include "Input/PopupMethodReply.h"
+#include "PixelFormat.h"
 #include "Rendering/DrawElementCoreTypes.h"
 #include "Rendering/SlateRendererTypes.h"
 #include "SlateGlobals.h"
@@ -485,6 +486,11 @@ public:
 	 */
 	virtual class FSlateShaderResource* GetViewportRenderTargetTexture() const = 0;
 
+	/**
+	 * Returns format for the scene of this viewport, only valid if this viewport actually has a scene, else PF_Unknown.
+	 */
+	virtual EPixelFormat GetSceneTargetFormat() const { return EPixelFormat::PF_Unknown; };
+
 	/**
 	 * Does the texture returned by GetViewportRenderTargetTexture only have an alpha channel?
 	 */
diff --git a/Engine/Source/Runtime/SlateRHIRenderer/Private/SlateRHIRenderer.cpp b/Engine/Source/Runtime/SlateRHIRenderer/Private/SlateRHIRenderer.cpp
--- a/Engine/Source/Runtime/SlateRHIRenderer/Private/SlateRHIRenderer.cpp
+++ b/Engine/Source/Runtime/SlateRHIRenderer/Private/SlateRHIRenderer.cpp
@@ -1825,12 +1825,10 @@ void FSlateRHIRenderer::DrawWindows_Private(FSlateDrawBuffer& WindowDrawBuffer)
 								: FIntPoint(ViewInfo->DesiredWidth, ViewInfo->DesiredHeight);
 
 							bool bHDREnabled = IsHDREnabled();
-							bool bIsPixelFormatCorrect = bHDREnabled
-								? SlatePostBuffer->GetFormat() == EPixelFormat::PF_FloatRGBA
-								: SlatePostBuffer->GetFormat() == EPixelFormat::PF_A2B10G10R10;
+							bool bIsPixelFormatCorrect = SlatePostBuffer->GetFormat() == Window->GetViewport()->GetSceneTargetFormat();
 							if (SlatePostBuffer->SizeX != SizeSlatePostRT.X || SlatePostBuffer->SizeY != SizeSlatePostRT.Y || !bIsPixelFormatCorrect)
 							{
-								SlatePostBuffer->InitCustomFormat(SizeSlatePostRT.X, SizeSlatePostRT.Y, bHDREnabled ? EPixelFormat::PF_FloatRGBA : EPixelFormat::PF_A2B10G10R10, true);
+								SlatePostBuffer->InitCustomFormat(SizeSlatePostRT.X, SizeSlatePostRT.Y, Window->GetViewport()->GetSceneTargetFormat(), true);
 							}
 
 							const FVector2D ElementWindowSize = ElementList.GetWindowSize();
@@ -1896,7 +1894,18 @@ void FSlateRHIRenderer::DrawWindows_Private(FSlateDrawBuffer& WindowDrawBuffer)
 							else
 							{
 								// Resize unused SlatePostRTs to 1x1.
-								SlatePostBuffer->InitCustomFormat(1, 1, IsHDREnabled() ? EPixelFormat::PF_FloatRGBA : EPixelFormat::PF_A2B10G10R10, true);
+								auto GetBackBufferFormat = [&]() 
+								{
+									// This path can execute when stopping PIE / exiting, so viewport may not be valid (above Viewport is always valid).
+									if (TSharedPtr<ISlateViewport> Viewport = Window->GetViewport())
+									{
+										return Viewport->GetSceneTargetFormat();
+									}
+
+									return IsHDREnabled() ? EPixelFormat::PF_FloatRGBA : EPixelFormat::PF_A2B10G10R10;
+								};
+
+								SlatePostBuffer->InitCustomFormat(1, 1, GetBackBufferFormat(), true);
 								bShrinkPostBufferRequested &= ~SlatePostBufferBit;
 							}
 						}
