Remove Slate.MemorylessDepthStencil option as it does not work correctly with RHIs that actually support memoryless textures

diff --git a/Engine/Source/Runtime/SlateRHIRenderer/Private/SlateRHIRenderer.cpp b/Engine/Source/Runtime/SlateRHIRenderer/Private/SlateRHIRenderer.cpp
--- a/Engine/Source/Runtime/SlateRHIRenderer/Private/SlateRHIRenderer.cpp
+++ b/Engine/Source/Runtime/SlateRHIRenderer/Private/SlateRHIRenderer.cpp
@@ -100,12 +100,6 @@ static TAutoConsoleVariable<float> CVarDrawToVRRenderTarget(
 	TEXT("If enabled while in VR. Slate UI will be drawn into the render target texture where the VR imagery for either eye was rendered, allow the viewer of the HMD to see the UI (for better or worse.)  This render target will then be cropped/scaled into the back buffer, if mirroring is enabled.  When disabled, Slate UI will be drawn on top of the backbuffer (not to the HMD) after the mirror texture has been cropped/scaled into the backbuffer."),
 	ECVF_RenderThreadSafe);
 
-static TAutoConsoleVariable<int32> CVarMemorylessDepthStencil(
-	TEXT("Slate.MemorylessDepthStencil"),
-	0,
-	TEXT("Whether to use memoryless DepthStencil target for Slate. Reduces memory usage and implies that DepthStencil state can't be preserved between Slate renderpasses"),
-	ECVF_RenderThreadSafe);
-
 static TAutoConsoleVariable<int32> CVarCopyBackbufferToSlatePostRenderTargets(
 	TEXT("Slate.CopyBackbufferToSlatePostRenderTargets"),
 	0,
@@ -181,12 +175,10 @@ void FViewportInfo::ConditionallyUpdateDepthBuffer(bool bInRequiresStencilTest,
 {
 	check(IsInRenderingThread());
 
-	bool bWantsMemorylessDepthStencil = (CVarMemorylessDepthStencil.GetValueOnAnyThread() != 0);
-
 	bool bDepthStencilStale =
 		bInRequiresStencilTest &&
 		(!bRequiresStencilTest ||
-		(DepthStencil.IsValid() && (DepthStencil->GetSizeX() != InWidth || DepthStencil->GetSizeY() != InHeight || IsMemorylessTexture(DepthStencil) != bWantsMemorylessDepthStencil)));
+		(DepthStencil.IsValid() && (DepthStencil->GetSizeX() != InWidth || DepthStencil->GetSizeY() != InHeight)));
 
 	bRequiresStencilTest = bInRequiresStencilTest;
 
@@ -204,12 +196,7 @@ void FViewportInfo::RecreateDepthBuffer_RenderThread()
 	if (bRequiresStencilTest)
 	{
 		ETextureCreateFlags TargetableTextureFlags = TexCreate_DepthStencilTargetable;
-		if (CVarMemorylessDepthStencil.GetValueOnAnyThread() != 0)
-		{
-			// Use Memoryless target, expecting that DepthStencil content is intermediate and can't be preserved between renderpasses
-			TargetableTextureFlags|= TexCreate_Memoryless;
-		}
-		
+	
 		const FRHITextureCreateDesc Desc =
 			FRHITextureCreateDesc::Create2D(TEXT("SlateViewportDepthStencil"))
 			.SetExtent(Width, Height)
@@ -223,16 +210,6 @@ void FViewportInfo::RecreateDepthBuffer_RenderThread()
 	}
 }
 
-bool IsMemorylessTexture(const FTextureRHIRef& Tex)
-{
-	if (Tex)
-	{
-		return EnumHasAnyFlags(Tex->GetFlags(), TexCreate_Memoryless);
-	}
-	return false;
-}
-
-
 FSlateRHIRenderer::FSlateRHIRenderer(TSharedRef<FSlateFontServices> InSlateFontServices, TSharedRef<FSlateRHIResourceManager> InResourceManager)
 	: FSlateRenderer(InSlateFontServices)
 	, EnqueuedWindowDrawBuffer(NULL)
@@ -989,7 +966,7 @@ void RenderSlateBatch(FTextureRHIRef SlateRenderTarget, bool bClear, bool bIsHDR
 	{
 		check(IsValidRef(ViewportInfo.DepthStencil));
 
-		ERenderTargetActions StencilAction = IsMemorylessTexture(ViewportInfo.DepthStencil) ? ERenderTargetActions::DontLoad_DontStore : ERenderTargetActions::DontLoad_Store;
+		ERenderTargetActions StencilAction = ERenderTargetActions::DontLoad_Store;
 		RPInfo.DepthStencilRenderTarget.Action = MakeDepthStencilTargetActions(ERenderTargetActions::DontLoad_DontStore, StencilAction);
 		RPInfo.DepthStencilRenderTarget.DepthStencilTarget = ViewportInfo.DepthStencil;
 		RPInfo.DepthStencilRenderTarget.ExclusiveDepthStencil = FExclusiveDepthStencil::DepthNop_StencilWrite;
@@ -1002,7 +979,7 @@ void RenderSlateBatch(FTextureRHIRef SlateRenderTarget, bool bClear, bool bIsHDR
 		if (ViewportInfo.bRequiresStencilTest)
 		{
 			// Reset the backbuffer as our color render target and also set a depth stencil buffer
-			ERenderTargetActions StencilAction = IsMemorylessTexture(ViewportInfo.DepthStencil) ? ERenderTargetActions::Clear_DontStore : ERenderTargetActions::Clear_Store;
+			ERenderTargetActions StencilAction = ERenderTargetActions::Clear_Store;
 			RPInfo.DepthStencilRenderTarget.Action = MakeDepthStencilTargetActions(ERenderTargetActions::Load_Store, StencilAction);
 			RPInfo.DepthStencilRenderTarget.DepthStencilTarget = ViewportInfo.DepthStencil;
 			RPInfo.DepthStencilRenderTarget.ExclusiveDepthStencil = FExclusiveDepthStencil::DepthWrite_StencilWrite;
diff --git a/Engine/Source/Runtime/SlateRHIRenderer/Private/SlateRHIRenderer.h b/Engine/Source/Runtime/SlateRHIRenderer/Private/SlateRHIRenderer.h
--- a/Engine/Source/Runtime/SlateRHIRenderer/Private/SlateRHIRenderer.h
+++ b/Engine/Source/Runtime/SlateRHIRenderer/Private/SlateRHIRenderer.h
@@ -155,8 +155,6 @@ struct FViewportInfo : public FRenderResource
 	}	
 };
 
-bool IsMemorylessTexture(const FTextureRHIRef& Tex);
-
 struct FFastPathRenderingDataCleanupList
 {
 	TArray<FSlateCachedFastPathRenderingData*, TInlineAllocator<20>> FastPathRenderingDataToRemove;
diff --git a/Engine/Source/Runtime/SlateRHIRenderer/Private/SlateRHIRenderingPolicy.cpp b/Engine/Source/Runtime/SlateRHIRenderer/Private/SlateRHIRenderingPolicy.cpp
--- a/Engine/Source/Runtime/SlateRHIRenderer/Private/SlateRHIRenderingPolicy.cpp
+++ b/Engine/Source/Runtime/SlateRHIRenderer/Private/SlateRHIRenderingPolicy.cpp
@@ -283,7 +283,7 @@ static bool UpdateScissorRect(
 					// #todo-renderpasses this is very gross. If/when this gets refactored we can detect a simple clear or batch up elements by rendertarget (and other stuff)
 					RHICmdList.EndRenderPass();
 					bDidRestartRenderpass = true;
-					ERenderTargetActions StencilAction = IsMemorylessTexture(DepthStencilTarget) ? ERenderTargetActions::DontLoad_DontStore : ERenderTargetActions::Load_Store;
+					ERenderTargetActions StencilAction = ERenderTargetActions::Load_Store;
 
 					RPInfo.DepthStencilRenderTarget.Action = MakeDepthStencilTargetActions(ERenderTargetActions::DontLoad_DontStore, StencilAction);
 					RPInfo.DepthStencilRenderTarget.DepthStencilTarget = DepthStencilTarget;
@@ -340,12 +340,6 @@ static bool UpdateScissorRect(
 					// Clear current stencil buffer, we use ELoad/EStore, because we need to keep the stencil around.
 					ERenderTargetLoadAction StencilLoadAction = bClearStencil ? ERenderTargetLoadAction::EClear : ERenderTargetLoadAction::ELoad;
 					ERenderTargetActions StencilAction = MakeRenderTargetActions(StencilLoadAction, ERenderTargetStoreAction::EStore);
-					if (IsMemorylessTexture(DepthStencilTarget))
-					{
-						// We can't preserve content for memoryless targets
-						StencilAction = bClearStencil ? ERenderTargetActions::Clear_DontStore : ERenderTargetActions::DontLoad_DontStore;
-					}
-
 					RPInfo.DepthStencilRenderTarget.Action = MakeDepthStencilTargetActions(ERenderTargetActions::DontLoad_DontStore, StencilAction);
 					RPInfo.DepthStencilRenderTarget.DepthStencilTarget = DepthStencilTarget;
 					RPInfo.DepthStencilRenderTarget.ExclusiveDepthStencil = FExclusiveDepthStencil::DepthNop_StencilWrite;
@@ -809,7 +803,7 @@ void FSlateRHIRenderingPolicy::DrawElements(
 			RPInfo.DepthStencilRenderTarget.DepthStencilTarget = DepthStencilTarget;
 			if (DepthStencilTarget)
 			{
-				RPInfo.DepthStencilRenderTarget.Action = IsMemorylessTexture(DepthStencilTarget) ? EDepthStencilTargetActions::DontLoad_DontStore : EDepthStencilTargetActions::LoadDepthStencil_StoreDepthStencil;
+				RPInfo.DepthStencilRenderTarget.Action = EDepthStencilTargetActions::LoadDepthStencil_StoreDepthStencil;
 				RPInfo.DepthStencilRenderTarget.ExclusiveDepthStencil = FExclusiveDepthStencil::DepthWrite_StencilWrite;
 			}
 			else
