Clean-up of Slate RHI renderer and conversion to RDG.
 - Removed legacy render pass management.
 - Merged individual window render passes into single RDG builder.
 - Cleaned up lots of dead code around HDR processing.
 - Moved render batch merging to run asynchronously on the CPU.
 - Ported UMG post buffer system and custom draw element paths to RDG.
 - Replaced expensive scene view creation with much simpler view uniform buffer path.

diff --git a/Engine/Plugins/Experimental/Mutable/Source/CustomizableObjectEditor/Private/MuCOE/SCustomizableObjectLayoutGrid.cpp b/Engine/Plugins/Experimental/Mutable/Source/CustomizableObjectEditor/Private/MuCOE/SCustomizableObjectLayoutGrid.cpp
--- a/Engine/Plugins/Experimental/Mutable/Source/CustomizableObjectEditor/Private/MuCOE/SCustomizableObjectLayoutGrid.cpp
+++ b/Engine/Plugins/Experimental/Mutable/Source/CustomizableObjectEditor/Private/MuCOE/SCustomizableObjectLayoutGrid.cpp
@@ -16,7 +16,7 @@
 #include "ThumbnailRendering/ThumbnailManager.h"
 #include "Engine/Texture2D.h"
 #include "TextureResource.h"
-
+#include "RenderGraphBuilder.h"
 
 class FExtender;
 class FPaintArgs;
@@ -37,15 +37,20 @@ public:
 	}
 
 	/** Sets the texture that this target renders to */
-	void SetRenderTargetTexture(FTextureRHIRef& InRHIRef)
+	void SetRenderTargetTexture(FRDGTexture* InTexture)
+	{
+		RDGTexture = InTexture;
+	}
+
+	FRDGTexture* GetRenderTargetTexture(FRDGBuilder&) const override
 	{
-		RenderTargetTextureRHI = InRHIRef;
+		return RDGTexture;
 	}
 
 	/** Clears the render target texture */
 	void ClearRenderTargetTexture()
 	{
-		RenderTargetTextureRHI.SafeRelease();
+		RDGTexture = nullptr;
 	}
 
 	/** Sets the viewport rect for the render target */
@@ -73,7 +78,7 @@ public:
 	}
 
 private:
-
+	FRDGTexture* RDGTexture = nullptr;
 	FIntRect ViewRect;
 	FIntRect ClippingRect;
 };
@@ -98,7 +103,7 @@ public:
 
 private:
 
-	virtual void Draw_RenderThread(FRHICommandListImmediate& RHICmdList, const void* InWindowBackBuffer, const FSlateCustomDrawParams& Params) override;
+	virtual void Draw_RenderThread(FRDGBuilder& GraphBuilder, const FDrawPassInputs& Inputs) override;
 
 	/** Basic function to draw a block in the canvas */
 	void DrawBlock(FBatchedElements* BatchedElements, const FHitProxyId HitProxyId, const FRect2D& BlockRect, FColor Color, UTexture2D* Mask = nullptr);
@@ -1352,17 +1357,17 @@ void FUVCanvasDrawer::SetLayoutMode(ELayoutGridMode Mode)
 }
 
 
-void FUVCanvasDrawer::Draw_RenderThread(class FRHICommandListImmediate& RHICmdList, const void* InWindowBackBuffer, const FSlateCustomDrawParams& Params)
+void FUVCanvasDrawer::Draw_RenderThread(FRDGBuilder& GraphBuilder, const FDrawPassInputs& Inputs)
 {
 	if (!Initialized)
 	{
 		return;
 	}
 
-	RenderTarget->SetRenderTargetTexture(*(FTextureRHIRef*)InWindowBackBuffer);
+	RenderTarget->SetRenderTargetTexture(Inputs.OutputTexture);
+
+	FCanvas& Canvas = *GraphBuilder.AllocObject<FCanvas>(RenderTarget, nullptr, FGameTime(), GMaxRHIFeatureLevel);
 
-	FCanvas Canvas(RenderTarget, nullptr, FGameTime(), GMaxRHIFeatureLevel);
-		
 	Canvas.SetRenderTargetRect(RenderTarget->GetViewRect());
 	Canvas.SetRenderTargetScissorRect(RenderTarget->GetClippingRect());
 
@@ -1545,10 +1550,9 @@ void FUVCanvasDrawer::Draw_RenderThread(class FRHICommandListImmediate& RHICmdLi
 		}
 	}
 
-	Canvas.Flush_RenderThread(RHICmdList, true);
+	Canvas.Flush_RenderThread(GraphBuilder, true);
 
 	RenderTarget->ClearRenderTargetTexture();
-	RHICmdList.SetScissorRect(false, 0, 0, 0, 0);
 }
 
 
diff --git a/Engine/Shaders/Private/CompositeUIPixelShader.usf b/Engine/Shaders/Private/CompositeUIPixelShader.usf
--- a/Engine/Shaders/Private/CompositeUIPixelShader.usf
+++ b/Engine/Shaders/Private/CompositeUIPixelShader.usf
@@ -10,20 +10,7 @@ CompositeUIPixelShader.usf: Filter pixel shader source.
 #define LUT_SCALING	1.05f
 #define LUT_SIZE	32.f
 
-// set this to 1 if you want to restore the legacy behaviour (ie apply ACES to the UI). A value of 0 keeps a consistent look of the UI between SDR and HDR mode
-#define USE_LUT_FOR_UI 0
-
 ///////////////////////////////////////////////////////////////////////////////////////
-uint OutputDevice;
-uint OutputGamut;
-float OutputMaxLuminance;
-
-float4 ACESMinMaxData;
-float4 ACESMidData;
-float4 ACESCoefsLow_0;
-float4 ACESCoefsHigh_0;
-float ACESCoefsLow_4;
-float ACESCoefsHigh_4;
 
 #if APPLY_COLOR_DEFICIENCY
 float ColorVisionDeficiencyType;
@@ -32,53 +19,6 @@ float bCorrectDeficiency;
 float bSimulateCorrectionWithDeficiency;
 #endif
 
-void GenerateLUTPS(
-	FWriteToSliceGeometryOutput Input,
-	out float4 OutColor : SV_Target0
-	)
-{
-	// Construct the neutral color from a 3D position volume texture
-	float2 UV = Input.Vertex.UV - float2(0.5f / LUT_SIZE, 0.5f / LUT_SIZE);
-	float3 Neutral = float3(UV * LUT_SIZE / (LUT_SIZE - 1), Input.LayerIndex / (LUT_SIZE - 1));
-
-	const float3x3 AP1_2_Output = OuputGamutMappingMatrix(OutputGamut);
-
-	float3 LinearColor = ST2084ToLinear(Neutral);
-
-	FACESTonemapParams LUTTsParams = ComputeACESTonemapParams(ACESMinMaxData, ACESMidData, ACESCoefsLow_0, ACESCoefsHigh_0, ACESCoefsLow_4, ACESCoefsHigh_4, 1.0, 1.0);
-
-	if( OutputDevice == TONEMAPPER_OUTPUT_ACES1000nitST2084 || OutputDevice == TONEMAPPER_OUTPUT_ACES1000nitScRGB)
-	{		
-		// ACES 1000nit transform with PQ/2084 encoding, user specified gamut 
-		float3 ODTColor = ACESOutputTransforms1000(LinearColor, LUTTsParams);
-		// do not apply scRGB here because LUT is not 16F in this format, and scRGB can output negative values
-		if (OutputDevice != TONEMAPPER_OUTPUT_ACES1000nitScRGB)
-		{
-			ODTColor = mul(AP1_2_Output, ODTColor);
-		}
-		OutColor.rgb = LinearToST2084(ODTColor);
-	}	
-	else if( OutputDevice == TONEMAPPER_OUTPUT_ACES2000nitST2084 || OutputDevice == TONEMAPPER_OUTPUT_ACES2000nitScRGB)
-	{		
-		// ACES 2000nit transform with PQ/2084 encoding, user specified gamut 
-		float3 ODTColor = ACESOutputTransforms2000(LinearColor, LUTTsParams);
-		// do not apply scRGB here because LUT is not 16F in this format, and scRGB can output negative values
-		// 
-		if (OutputDevice != TONEMAPPER_OUTPUT_ACES2000nitScRGB)
-		{
-			ODTColor = mul(AP1_2_Output, ODTColor);
-		}
-		OutColor.rgb = LinearToST2084(ODTColor);
-	}
-	else
-	{
-		OutColor.rgb = Neutral;
-	}
-
-	OutColor.rgb = OutColor.rgb / LUT_SCALING;
-	OutColor.a = 0;
-}
-
 ///////////////////////////////////////////////////////////////////////////////////////
 #if PLATFORM_SUPPORTS_RENDERTARGET_WRITE_MASK
 Texture2D<uint> UIWriteMaskTexture;
@@ -90,9 +30,6 @@ SamplerState UISampler;
 Texture2D SceneTexture;
 SamplerState SceneSampler;
 
-Texture3D ColorSpaceLUT;
-SamplerState ColorSpaceLUTSampler;
-
 float UILevel;
 float UILuminance;
 
@@ -114,13 +51,6 @@ float3 LinearizeColor(float3 EncodedColor)
 
 static const float ScRGBScaleFactor = 80.0f; // used to convert between ScRGB and nits
 
-float3 ColorLookupTable(float3 LinearColor)
-{
-	float3 LUTEncodedColor = LinearToST2084(LinearColor);
-	float3 UVW = LUTEncodedColor * ((LUT_SIZE - 1) / LUT_SIZE) + (0.5f / LUT_SIZE);
-	return ColorSpaceLUT.SampleLevel(ColorSpaceLUTSampler, UVW, 0).rgb * LUT_SCALING;
-}
-
 void HDRBackBufferConvertPS(
 	FScreenVertexOutput Input,
 	out float3 OutColor : SV_Target0
@@ -132,28 +62,6 @@ void HDRBackBufferConvertPS(
 	OutColor.rgb = SceneColor;
 }
 
-// Given an input using Rec709 primaries and a sRGB tranfer function where
-// 1.0 maps to paperwhitenits defined in editor. Remap to Rec2020 with PQ transfer function.
-void HDREditorConvert(
-	FScreenVertexOutput Input,
-	out float4 OutColor : SV_Target0
-	)
-{
-	OutColor = 0;
-	float3 SceneColor = Texture2DSample(SceneTexture, SceneSampler, Input.UV).xyz ;
-
-	// Get color into Rec2020 space
-	SceneColor = sRGBToLinear( SceneColor) ;
-	SceneColor = SceneColor * UILevel ; 
-
-	SceneColor = mul(sRGB_2_XYZ_MAT, SceneColor);
-	SceneColor = mul(XYZ_2_Rec2020_MAT, SceneColor);
-
-	// Do PQ Curve
-	SceneColor = LinearToST2084( SceneColor ); 
-	OutColor.rgb = SceneColor;
-}
-
 float3 ComposeUIAndScene(float3 SceneColor, float4 InUIColor, float InUILevel)
 {
 	BRANCH
@@ -198,20 +106,7 @@ float4 ComputeHDRUIColor(float2 InputUV)
 	const float3x3 sRGB_2_Rec2020 = mul(XYZ_2_Rec2020_MAT, sRGB_2_XYZ_MAT);
 	float4 UIColor = Texture2DSample(UITexture, UISampler, InputUV);
 	UIColor.xyz = LinearizeColor(UIColor.xyz);
-#if USE_LUT_FOR_UI
-	UIColor.xyz = ColorLookupTable(UIColor.xyz);
-	// PQ -> Linear
-#if !SCRGB_ENCODING
-	UIColor.xyz = ST2084ToLinear(UIColor.xyz);
-#else
-	UIColor.xyz = ST2084ToScRGB(UIColor.xyz, OutputDevice, OutputMaxLuminance);
 	UIColor.xyz = mul(sRGB_2_Rec2020, UIColor.xyz) * UILuminance;
-#endif
-
-#else // #if USE_LUT_FOR_UI
-	UIColor.xyz = mul(sRGB_2_Rec2020, UIColor.xyz) * UILuminance;
-#endif // #if USE_LUT_FOR_UI
-
 	return UIColor;
 }
 
@@ -330,14 +225,11 @@ void CompositeUICS(uint3 DispatchThreadId : SV_DispatchThreadID)
 
 #endif
 
-#if BLIT_UI_TO_HDR
+#if COMPOSITE_UI_FOR_BLUR_PS
 
 float2 UITextureSize;
 
-void BlitUIToHDRPS(
-	FScreenVertexOutput Input,
-	out float4 OutColor : SV_Target0
-	)
+void CompositeUIForBlur(FScreenVertexOutput Input, out float4 OutColor : SV_Target0)
 {
 	float4 UIColor = ComputeHDRUIColor(Input.UV);
 #if PLATFORM_SUPPORTS_RENDERTARGET_WRITE_MASK
@@ -356,4 +248,4 @@ void BlitUIToHDRPS(
 	OutColor.a = 1;
 }
 
-#endif
+#endif
\ No newline at end of file
diff --git a/Engine/Shaders/Private/SlateMaskingShader.usf b/Engine/Shaders/Private/SlateMaskingShader.usf
--- a/Engine/Shaders/Private/SlateMaskingShader.usf
+++ b/Engine/Shaders/Private/SlateMaskingShader.usf
@@ -2,13 +2,6 @@
 
 #include "Common.ush"
 
-float4x4 ViewProjection;
-
-
-// MaskRect[0].xy = TopLeft
-// MaskRect[0].zw = TopRight
-// MaskRect[1].xy = BottomLeft
-// MaskRect[1].zw = BottomRight
 float4 MaskRectPacked[2];
 
 struct VertexToPixelInterpolants
@@ -19,17 +12,12 @@ struct VertexToPixelInterpolants
 VertexToPixelInterpolants MainVS(in uint InPosition : ATTRIBUTE0)
 {
 	VertexToPixelInterpolants VOut = (VertexToPixelInterpolants)0;
-
 	const float2 MaskRect[4] = { MaskRectPacked[0].xy, MaskRectPacked[0].zw, MaskRectPacked[1].xy, MaskRectPacked[1].zw };
-
-    float4 WorldPosition = float4(MaskRect[InPosition], 0, 1);
-
-	VOut.Position = mul(WorldPosition, ViewProjection);
-
+	VOut.Position = float4(MaskRect[InPosition], 0, 1);
 	return VOut;
 }
 
 float4 MainPS(VertexToPixelInterpolants VIn) : SV_Target0
 {
-    return 1;
-}
+	return 1;
+}
\ No newline at end of file
diff --git a/Engine/Shaders/Private/SlatePostProcessColorDeficiencyPixelShader.usf b/Engine/Shaders/Private/SlatePostProcessColorDeficiencyPixelShader.usf
--- a/Engine/Shaders/Private/SlatePostProcessColorDeficiencyPixelShader.usf
+++ b/Engine/Shaders/Private/SlatePostProcessColorDeficiencyPixelShader.usf
@@ -17,11 +17,6 @@ float bSimulateCorrectionWithDeficiency;
 float4 ColorDeficiencyMain(FScreenVertexOutput Input) : SV_Target0
 {
 	float4 SourceColor = Texture2DSample(ElementTexture, ElementTextureSampler, Input.UV);
-	// Have to convert to linear color space first - so they say, but it doesn't appear to matter much.
-	//SourceColor.rgb = sRGBToLinear(SourceColor.rgb);
-
 	float3 OutRGB = ColorDeficiency(SourceColor.rgb, ColorVisionDeficiencyType, ColorVisionDeficiencySeverity, bCorrectDeficiency, bSimulateCorrectionWithDeficiency);
-	//OutRGB = LinearToSrgb(OutRGB);
-
 	return float4(OutRGB, SourceColor.a);
-}
+}
\ No newline at end of file
diff --git a/Engine/Shaders/Private/SlatePostProcessPixelShader.usf b/Engine/Shaders/Private/SlatePostProcessPixelShader.usf
--- a/Engine/Shaders/Private/SlatePostProcessPixelShader.usf
+++ b/Engine/Shaders/Private/SlatePostProcessPixelShader.usf
@@ -4,25 +4,14 @@
 #include "SlateShaderCommon.ush"
 #include "GammaCorrectionCommon.ush"
 
-#define BILINEAR_FILTER_METHOD 1
 #define MAX_SAMPLES 127
 
-#if BILINEAR_FILTER_METHOD
-
 // Weigts and offsets are packed into 4 floats (Weight, Offset, Weight, Offset)
 float4 WeightAndOffsets[MAX_SAMPLES/2];
-#else
-
-#define MAX_SAMPLES 127
-// Weigts and offsets are packed into 4 floats (Weight, Offset, Weight, Offset)
-float4 WeightAndOffsets[MAX_SAMPLES];
-
-#endif
 
 /** Blur sample count */
 int SampleCount;
 
-
 Texture2D ElementTexture;
 SamplerState ElementTextureSampler;
 
@@ -46,10 +35,6 @@ float4 GetSample(float Weight, float Offset, float2 UV)
 
 float4 GaussianBlurMain( FScreenVertexOutput Input ) : SV_Target0
 {
-#if 0
-	float4 OutColor = Texture2DSample(ElementTexture, ElementTextureSampler, Input.UV);
-#else
-
 	float4 OutColor = Texture2DSample(ElementTexture, ElementTextureSampler, clamp(Input.UV, UVBounds.xy, UVBounds.zw)) * WeightAndOffsets[0].x;
 	// First offset is in zw
 	{
@@ -77,17 +62,12 @@ float4 GaussianBlurMain( FScreenVertexOutput Input ) : SV_Target0
 		}
 
 	}
-#endif
 
 	return float4(OutColor.rgb, 1);
 }
 
 float4 DownsampleMain(FScreenVertexOutput Input) : SV_Target0
 {
-#if 0
-	return Texture2DSample(ElementTexture, ElementTextureSampler, Input.UV);
-#else
-
 	float2 UV[4];
 
 	float2 MinUV = UVBounds.xy;
@@ -107,17 +87,20 @@ float4 DownsampleMain(FScreenVertexOutput Input) : SV_Target0
 	}
 
 	return float4(Sample[0] + Sample[1] + Sample[2] + Sample[3]) * 0.25f;
-#endif
 }
 
 float2 Map(float2 value, float2 min1, float2 max1, float2 min2, float2 max2)
 {
   return min2 + (value - min1) * (max2 - min2) / (max1 - min1);
 }
- 
+
+#define UPSAMPLE_OUTPUT_FORMAT_SDR 0
+#define UPSAMPLE_OUTPUT_FORMAT_HDR_SCRGB 1
+#define UPSAMPLE_OUTPUT_FORMAT_HDR_PQ10 2
+
 void UpsampleMain(FScreenVertexOutput Input
 	, out float4 Color0 : SV_Target0
-#if OUTPUT_TO_UI_TARGET
+#if UPSAMPLE_OUTPUT_FORMAT != UPSAMPLE_OUTPUT_FORMAT_SDR
 	, out float4 Color1 : SV_Target1
 #endif
 )
@@ -129,8 +112,7 @@ void UpsampleMain(FScreenVertexOutput Input
 	const float4 FillColor = Texture2DSample(ElementTexture, ElementTextureSampler, Input.UV);
 	Color0 = GetRoundedBoxElementColorInternal(LocalSize, Map(UV, float2(0,0), ShaderParams.zw, float2(0,0),float2(1,1)), Thickness, CornerRadii, FillColor, FillColor);
 
-#if OUTPUT_TO_UI_TARGET
-
+#if UPSAMPLE_OUTPUT_FORMAT != UPSAMPLE_OUTPUT_FORMAT_SDR
 	static const float3x3 XYZ_2_sRGB_MAT =
 	{
 		 3.2409699419, -1.5373831776, -0.4986107603,
@@ -149,7 +131,7 @@ void UpsampleMain(FScreenVertexOutput Input
 	static const float ScRGBScaleFactor = 80.0f; // used to convert between ScRGB and nits
 	Color0.xyz *= ScRGBScaleFactor;
 
-#if SCRGB_ENCODING
+#if UPSAMPLE_OUTPUT_FORMAT == UPSAMPLE_OUTPUT_FORMAT_HDR_SCRGB
 	const float3x3 Rec2020_2_sRGB = mul(XYZ_2_sRGB_MAT, Rec2020_2_XYZ_MAT);
 	Color0.xyz = mul(Rec2020_2_sRGB, Color0.xyz) / ScRGBScaleFactor;
 #else
@@ -158,6 +140,5 @@ void UpsampleMain(FScreenVertexOutput Input
 
 	// Erase the data that was previously in the UI now that the background blur is done
 	Color1 = float4(0, 0, 0, 0);
-
 #endif
 }
\ No newline at end of file
diff --git a/Engine/Source/Developer/StandaloneRenderer/Private/OpenGL/SlateOpenGLRenderingPolicy.cpp b/Engine/Source/Developer/StandaloneRenderer/Private/OpenGL/SlateOpenGLRenderingPolicy.cpp
--- a/Engine/Source/Developer/StandaloneRenderer/Private/OpenGL/SlateOpenGLRenderingPolicy.cpp
+++ b/Engine/Source/Developer/StandaloneRenderer/Private/OpenGL/SlateOpenGLRenderingPolicy.cpp
@@ -86,8 +86,6 @@ void FSlateOpenGLRenderingPolicy::ReleaseResources()
  */
 void FSlateOpenGLRenderingPolicy::BuildRenderingBuffers(FSlateBatchData& InBatchData)
 {
-	InBatchData.MergeRenderBatches();
-
 	if( InBatchData.GetRenderBatches().Num() > 0 )
 	{
 		const FSlateVertexArray& FinalVertexData = InBatchData.GetFinalVertexData();
diff --git a/Engine/Source/Developer/StandaloneRenderer/Private/Windows/D3D/SlateD3DRenderingPolicy.cpp b/Engine/Source/Developer/StandaloneRenderer/Private/Windows/D3D/SlateD3DRenderingPolicy.cpp
--- a/Engine/Source/Developer/StandaloneRenderer/Private/Windows/D3D/SlateD3DRenderingPolicy.cpp
+++ b/Engine/Source/Developer/StandaloneRenderer/Private/Windows/D3D/SlateD3DRenderingPolicy.cpp
@@ -211,9 +211,6 @@ void FSlateD3D11RenderingPolicy::ReleaseResources()
 
 void FSlateD3D11RenderingPolicy::BuildRenderingBuffers( FSlateBatchData& InBatchData )
 {	
-	// Merge render batches together to form final rendering calls
-	InBatchData.MergeRenderBatches();
-
 	if (!VertexBuffer.GetResource().IsValid() ||
 		!IndexBuffer.GetResource().IsValid())
 	{
diff --git a/Engine/Source/Editor/GraphEditor/Private/MaterialNodes/SGraphNodeMaterialBase.cpp b/Engine/Source/Editor/GraphEditor/Private/MaterialNodes/SGraphNodeMaterialBase.cpp
--- a/Engine/Source/Editor/GraphEditor/Private/MaterialNodes/SGraphNodeMaterialBase.cpp
+++ b/Engine/Source/Editor/GraphEditor/Private/MaterialNodes/SGraphNodeMaterialBase.cpp
@@ -65,6 +65,7 @@
 #include "Rendering/SubstrateMaterialShared.h"
 #include "SGraphSubstrateMaterial.h"
 #include "MaterialShared.h"
+#include "RenderGraphBuilder.h"
 
 class FWidgetStyle;
 class SWidget;
@@ -87,15 +88,20 @@ public:
 	}
 
 	/** Sets the texture that this target renders to */
-	void SetRenderTargetTexture( FTextureRHIRef& InRHIRef )
+	void SetRenderTargetTexture(FRDGTexture* Texture)
 	{
-		RenderTargetTextureRHI = InRHIRef;
+		RDGTexture = Texture;
 	}
 
 	/** Clears the render target texture */
 	void ClearRenderTargetTexture()
 	{
-		RenderTargetTextureRHI.SafeRelease();
+		RDGTexture = nullptr;
+	}
+
+	FRDGTextureRef GetRenderTargetTexture(FRDGBuilder&) const override
+	{
+		return RDGTexture;
 	}
 
 	/** Sets the viewport rect for the render target */
@@ -123,6 +129,7 @@ public:
 	}
 
 private:
+	FRDGTexture* RDGTexture = nullptr;
 	FIntRect ViewRect;
 	FIntRect ClippingRect;
 };
@@ -264,29 +271,28 @@ void FPreviewElement::UpdateExpressionPreview(UMaterialGraphNode* MaterialNode)
 	);
 }
 
-void FPreviewElement::Draw_RenderThread(FRHICommandListImmediate& RHICmdList, const void* InWindowBackBuffer, const FSlateCustomDrawParams& Params)
+void FPreviewElement::Draw_RenderThread(FRDGBuilder& GraphBuilder, const FDrawPassInputs& Inputs)
 {
 	if(ExpressionPreview)
 	{
-		RenderTarget->SetRenderTargetTexture(*(FTextureRHIRef*)InWindowBackBuffer);
+		RenderTarget->SetRenderTargetTexture(Inputs.OutputTexture);
 		{
 			// Check realtime mode for whether to pass current time to canvas
 			double CurrentTime = bIsRealtime ? (FApp::GetCurrentTime() - GStartTime) : 0.0;
 			float DeltaTime = bIsRealtime ? FApp::GetDeltaTime() : 0.0f;
 
-			FCanvas Canvas(RenderTarget, NULL, FGameTime::CreateUndilated(CurrentTime, DeltaTime), GMaxRHIFeatureLevel);
+			FCanvas* Canvas = GraphBuilder.AllocObject<FCanvas>(RenderTarget, nullptr, FGameTime::CreateUndilated(CurrentTime, DeltaTime), GMaxRHIFeatureLevel);
 			{
-				Canvas.SetAllowedModes(0);
-				Canvas.SetRenderTargetRect(RenderTarget->GetViewRect());
-				Canvas.SetRenderTargetScissorRect(RenderTarget->GetClippingRect());
+				Canvas->SetAllowedModes(0);
+				Canvas->SetRenderTargetRect(RenderTarget->GetViewRect());
+				Canvas->SetRenderTargetScissorRect(RenderTarget->GetClippingRect());
 
 				FCanvasTileItem TileItem(FVector2D::ZeroVector, ExpressionPreview, RenderTarget->GetSizeXY());
-				Canvas.DrawItem(TileItem);
+				Canvas->DrawItem(TileItem);
+				Canvas->Flush_RenderThread(GraphBuilder, true);
 			}
-			Canvas.Flush_RenderThread(RHICmdList, true);
 		}
 		RenderTarget->ClearRenderTargetTexture();
-		RHICmdList.SetScissorRect(false, 0, 0, 0, 0);
 	}
 }
 
diff --git a/Engine/Source/Editor/GraphEditor/Public/MaterialNodes/SGraphNodeMaterialBase.h b/Engine/Source/Editor/GraphEditor/Public/MaterialNodes/SGraphNodeMaterialBase.h
--- a/Engine/Source/Editor/GraphEditor/Public/MaterialNodes/SGraphNodeMaterialBase.h
+++ b/Engine/Source/Editor/GraphEditor/Public/MaterialNodes/SGraphNodeMaterialBase.h
@@ -77,7 +77,7 @@ private:
 	/**
 	 * ICustomSlateElement interface 
 	 */
-	virtual void Draw_RenderThread(FRHICommandListImmediate& RHICmdList, const void* InWindowBackBuffer, const FSlateCustomDrawParams& Params) override;
+	virtual void Draw_RenderThread(FRDGBuilder& GraphBuilder, const FDrawPassInputs& Inputs) override;
 
 private:
 	/** Render target that the canvas renders to */
diff --git a/Engine/Source/Runtime/Engine/Private/Slate/DebugCanvas.cpp b/Engine/Source/Runtime/Engine/Private/Slate/DebugCanvas.cpp
--- a/Engine/Source/Runtime/Engine/Private/Slate/DebugCanvas.cpp
+++ b/Engine/Source/Runtime/Engine/Private/Slate/DebugCanvas.cpp
@@ -16,6 +16,7 @@
 #include "IHeadMountedDisplay.h"
 #include "RenderTargetPool.h"
 #include "ViewportClient.h"
+#include "RenderGraphUtils.h"
 
 /**
  * Simple representation of the backbuffer that the debug canvas renders to
@@ -31,29 +32,41 @@ public:
 	}
 
 	/** Sets the texture that this target renders to */
-	void SetRenderTargetTexture( FRHITexture* InRHIRef )
+	void SetRenderTargetTexture(FRDGTexture* InRDGTexture)
 	{
-		RenderTargetTextureRHI = InRHIRef;
+		RDGTexture = InRDGTexture;
 	}
 
 	/** Clears the render target texture */
 	void ClearRenderTargetTexture()
 	{
-		RenderTargetTextureRHI.SafeRelease();
+		RDGTexture = nullptr;
 	}
 
-	/** Sets the viewport rect for the render target */
-	void SetViewRect( const FIntRect& InViewRect ) 
-	{ 
+	const FTextureRHIRef& GetRenderTargetTexture() const override
+	{
+		static FTextureRHIRef NullRef;
+		return NullRef;
+	}
+
+	FRDGTextureRef GetRenderTargetTexture(FRDGBuilder&) const override
+	{
+		return RDGTexture;
+	}
+
+	void SetViewRect(const FIntRect& InViewRect)
+	{
 		ViewRect = InViewRect;
 	}
 
 	/** Gets the viewport rect for the render target */
 	const FIntRect& GetViewRect() const 
 	{
-		return ViewRect; 
+		return ViewRect;
 	}
+
 private:
+	FRDGTexture* RDGTexture = nullptr;
 	FIntRect ViewRect;
 };
 
@@ -234,40 +247,24 @@ void FDebugCanvasDrawer::InitDebugCanvas(FViewportClient* ViewportClient, UWorld
 	}
 }
 
-void FDebugCanvasDrawer::Draw_RenderThread(FRHICommandListImmediate& RHICmdList, const void* InWindowBackBuffer, const FSlateCustomDrawParams& Params)
+void FDebugCanvasDrawer::Draw_RenderThread(FRDGBuilder& GraphBuilder, const FDrawPassInputs& Inputs)
 {
-	check( IsInRenderingThread() );
-	check(RHICmdList.IsOutsideRenderPass());
-
-	SCOPED_DRAW_EVENT(RHICmdList, DrawDebugCanvas);
+	RDG_EVENT_SCOPE(GraphBuilder, "DrawDebugCanvas");
+	TRACE_CPUPROFILER_EVENT_SCOPE(DrawDebugCanvas);
 
-	QUICK_SCOPE_CYCLE_COUNTER(Stat_DrawDebugCanvas);
-	if( RenderThreadCanvas.IsValid() )
+	if (RenderThreadCanvas.IsValid())
 	{
-		FRHITexture* RT = InWindowBackBuffer != nullptr ? ((FTextureRHIRef*)InWindowBackBuffer)->GetReference() : nullptr;
-		FTextureRHIRef HMDSwapchain = nullptr;	// Attention: RT could point to HMDSwapchain later.
+		FRDGTexture* OutputTexture = Inputs.OutputTexture;
+
 		if (RenderThreadCanvas->IsUsingInternalTexture())
 		{
-			if (LayerTexture && RenderThreadCanvas->GetParentCanvasSize() != LayerTexture->GetDesc().Extent)
-			{
-				LayerTexture.SafeRelease();
-			}
-
-			if (!LayerTexture)
-			{
-				// Set TexCreate_NoFastClear because the fast CMASK clear was not working on ps4.
-				FPooledRenderTargetDesc Desc(FPooledRenderTargetDesc::Create2DDesc(RenderThreadCanvas->GetParentCanvasSize(), PF_B8G8R8A8, FClearValueBinding(), ETextureCreateFlags::None, ETextureCreateFlags::RenderTargetable | ETextureCreateFlags::NoFastClear, false));
-				Desc.DebugName = TEXT("DebugCanvasLayerTexture");
-				GRenderTargetPool.FindFreeElement(RHICmdList, Desc, LayerTexture, TEXT("DebugCanvasLayerTexture"));
-				UE_LOG(LogProfilingDebugging, Log, TEXT("Allocated a %d x %d texture for HMD canvas layer"), RenderThreadCanvas->GetParentCanvasSize().X, RenderThreadCanvas->GetParentCanvasSize().Y);
-			}
+			FTextureRHIRef HMDSwapChain;
 
 			IStereoLayers* const StereoLayers = (GEngine && GEngine->IsStereoscopic3D() && GEngine->StereoRenderingDevice.IsValid()) ? GEngine->StereoRenderingDevice->GetStereoLayers() : nullptr;
-
-			FTextureRHIRef HMDNull = nullptr;
 			if (StereoLayers)
 			{
-				StereoLayers->GetAllocatedTexture(LayerID, HMDSwapchain, HMDNull);
+				FTextureRHIRef HMDNull;
+				StereoLayers->GetAllocatedTexture(LayerID, HMDSwapChain, HMDNull);
 
 				// If drawing to a layer tell the spectator screen controller to copy that layer to the spectator screen.
 				if (StereoLayers->ShouldCopyDebugLayersToSpectatorScreen() && LayerID != INVALID_LAYER_ID && GEngine && GEngine->XRSystem)
@@ -283,21 +280,45 @@ void FDebugCanvasDrawer::Draw_RenderThread(FRHICommandListImmediate& RHICmdList,
 					}
 				}
 			}
-			RT = reinterpret_cast<FRHITexture*>(HMDSwapchain == nullptr ? LayerTexture->GetRHI() : HMDSwapchain.GetReference());
+
+			if (LayerTexture && RenderThreadCanvas->GetParentCanvasSize() != LayerTexture->GetDesc().Extent)
+			{
+				LayerTexture.SafeRelease();
+			}
+
+			if (HMDSwapChain)
+			{
+				OutputTexture = RegisterExternalTexture(GraphBuilder, HMDSwapChain, TEXT("HMDSwapChainTexture"));
+			}
+			else if (LayerTexture)
+			{
+				OutputTexture = GraphBuilder.RegisterExternalTexture(LayerTexture);
+			}
+			else
+			{
+				OutputTexture = GraphBuilder.CreateTexture(
+					FRDGTextureDesc::Create2D(RenderThreadCanvas->GetParentCanvasSize(), PF_B8G8R8A8, FClearValueBinding(), ETextureCreateFlags::RenderTargetable),
+					TEXT("DebugCanvasLayerTexture"));
+
+				LayerTexture = GraphBuilder.ConvertToExternalTexture(OutputTexture);
+
+				UE_LOG(LogProfilingDebugging, Log, TEXT("Allocated a %d x %d texture for HMD canvas layer"), RenderThreadCanvas->GetParentCanvasSize().X, RenderThreadCanvas->GetParentCanvasSize().Y);
+			}
 		}
-		RHICmdList.Transition(FRHITransitionInfo(RT, ERHIAccess::Unknown, ERHIAccess::RTV));
-		RenderTarget->SetRenderTargetTexture(RT);
 
-		if (RenderThreadCanvas->IsScaledToRenderTarget() && RT) 
+		RenderTarget->SetRenderTargetTexture(OutputTexture);
+
+		if (RenderThreadCanvas->IsUsingInternalTexture())
 		{
-			RenderThreadCanvas->SetRenderTargetRect( FIntRect(0, 0, (RT)->GetSizeX(), (RT)->GetSizeY()) );
+			RenderThreadCanvas->SetRenderTargetRect(FIntRect(FIntPoint::ZeroValue, OutputTexture->Desc.Extent));
 		}
 		else
 		{
-			RenderThreadCanvas->SetRenderTargetRect( RenderTarget->GetViewRect() );
+			RenderThreadCanvas->SetRenderTargetRect(RenderTarget->GetViewRect());
 		}
 
-		RenderThreadCanvas->Flush_RenderThread(RHICmdList, true);
+		RenderThreadCanvas->Flush_RenderThread(GraphBuilder, true);
+
 		RenderTarget->ClearRenderTargetTexture();
 	}
 }
diff --git a/Engine/Source/Runtime/Engine/Private/Slate/SceneViewport.cpp b/Engine/Source/Runtime/Engine/Private/Slate/SceneViewport.cpp
--- a/Engine/Source/Runtime/Engine/Private/Slate/SceneViewport.cpp
+++ b/Engine/Source/Runtime/Engine/Private/Slate/SceneViewport.cpp
@@ -1992,33 +1992,12 @@ void FSceneViewport::SwapStatCommands( const FSceneViewport& OtherViewport )
 
 /** Queue an update to the Window's RT on the Renderthread */
 void FSceneViewport::WindowRenderTargetUpdate(FSlateRenderer* Renderer, SWindow* Window)
-{	
+{
 	check(IsInGameThread());
-	if (Renderer)
+	if (Renderer && Window)
 	{
-		if (UseSeparateRenderTarget())
-		{
-			if (Window)
-			{
-				// We need to pass a texture to the renderer only for stereo rendering. Otherwise, Editor will be rendered incorrectly.
-				if (GEngine->IsStereoscopic3D(this))
-				{
-					//todo: mw Make this function take an FSlateTexture* rather than a void*
-					Renderer->SetWindowRenderTarget(*Window, static_cast<IViewportRenderTargetProvider*>(this));
-				}
-				else
-				{
-					Renderer->SetWindowRenderTarget(*Window, nullptr);
-				}
-			}
-		}
-		else
-		{
-			if (Window)
-			{
-				Renderer->SetWindowRenderTarget(*Window, nullptr);
-			}
-		}
+		// The viewport provider interface is ONLY used for stereo VR compositing to blit the intermediate viewport target over to the VR swap chain.
+		Renderer->SetWindowRenderTarget(*Window, UseSeparateRenderTarget() && GEngine->IsStereoscopic3D(this) ? this : nullptr);
 	}
 }
 
diff --git a/Engine/Source/Runtime/Engine/Public/Slate/DebugCanvas.h b/Engine/Source/Runtime/Engine/Public/Slate/DebugCanvas.h
--- a/Engine/Source/Runtime/Engine/Public/Slate/DebugCanvas.h
+++ b/Engine/Source/Runtime/Engine/Public/Slate/DebugCanvas.h
@@ -74,7 +74,7 @@ private:
 	/**
 	 * ICustomSlateElement interface 
 	 */
-	virtual void Draw_RenderThread(FRHICommandListImmediate& RHICmdList, const void* InWindowBackBuffer, const FSlateCustomDrawParams& Params) override;
+	virtual void Draw_RenderThread(FRDGBuilder& GraphBuilder, const FDrawPassInputs& Inputs) override;
 
 	/**
 	 * Deletes the rendering thread canvas 
diff --git a/Engine/Source/Runtime/Renderer/Public/ScreenPass.h b/Engine/Source/Runtime/Renderer/Public/ScreenPass.h
--- a/Engine/Source/Runtime/Renderer/Public/ScreenPass.h
+++ b/Engine/Source/Runtime/Renderer/Public/ScreenPass.h
@@ -98,6 +98,16 @@ struct FScreenPassRenderTarget : public FScreenPassTexture
 
 	FScreenPassRenderTarget() = default;
 
+	explicit FScreenPassRenderTarget(FScreenPassTexture InTexture)
+		: FScreenPassTexture(InTexture)
+	{
+		// Derive the load action based on whether we are rendering to the entire texture or not.
+		if (ViewRect.Min != FIntPoint::ZeroValue || Texture && Texture->Desc.Extent != ViewRect.Max)
+		{
+			LoadAction = ERenderTargetLoadAction::ELoad;
+		}
+	}
+
 	FScreenPassRenderTarget(FScreenPassTexture InTexture, ERenderTargetLoadAction InLoadAction)
 		: FScreenPassTexture(InTexture)
 		, LoadAction(InLoadAction)
diff --git a/Engine/Source/Runtime/SlateCore/Private/Rendering/DrawElements.cpp b/Engine/Source/Runtime/SlateCore/Private/Rendering/DrawElements.cpp
--- a/Engine/Source/Runtime/SlateCore/Private/Rendering/DrawElements.cpp
+++ b/Engine/Source/Runtime/SlateCore/Private/Rendering/DrawElements.cpp
@@ -53,27 +53,16 @@ FSlateWindowElementList::FSlateWindowElementList(const TSharedPtr<SWindow>& InPa
 	, ResolveToDeferredIndex()
 	, WindowSize(FVector2f(0.0f, 0.0f))
 	, bIsInGameLayer(false)
-	//, bReportReferences(true)
 {
 	if (InPaintWindow.IsValid())
 	{
 		WindowSize = UE::Slate::CastToVector2f(InPaintWindow->GetSizeInScreen());
 	}
-
-	// Only keep UObject resources alive if this window element list is born on the game thread.
-/*
-	if (IsInGameThread())
-	{
-		ResourceGCRoot = MakeUnique<FWindowElementGCObject>(this);
-	}*/
 }
 
 FSlateWindowElementList::~FSlateWindowElementList()
 {
-	/*if (ResourceGCRoot.IsValid())
-	{
-		ResourceGCRoot->ClearOwner();
-	}*/
+	FinishMergeRenderBatches();
 }
 
 void FSlateWindowElementList::SetIsInGameLayer(bool bInGameLayer)
@@ -86,6 +75,25 @@ bool FSlateWindowElementList::GetIsInGameLayer()
 	return bIsInGameLayer;
 }
 
+void FSlateWindowElementList::StartMergeRenderBatches()
+{
+	if (!BatchData.GetRenderBatches().IsEmpty() || !BatchDataHDR.GetRenderBatches().IsEmpty())
+	{
+		MergeBatchDataTask = UE::Tasks::Launch(TEXT("Slate::MergeRenderBatches"), [this]
+		{
+			TRACE_CPUPROFILER_EVENT_SCOPE("Slate::MergeRenderBatches")
+			BatchData.MergeRenderBatches();
+			BatchDataHDR.MergeRenderBatches();
+		});
+	}
+}
+
+void FSlateWindowElementList::FinishMergeRenderBatches()
+{
+	MergeBatchDataTask.Wait();
+	MergeBatchDataTask = {};
+}
+
 FSlateWindowElementList::FDeferredPaint::FDeferredPaint( const TSharedRef<const SWidget>& InWidgetToPaint, const FPaintArgs& InArgs, const FGeometry InAllottedGeometry, const FWidgetStyle& InWidgetStyle, bool InParentEnabled )
 	: WidgetToPaintPtr( InWidgetToPaint )
 	, Args( InArgs )
@@ -198,38 +206,6 @@ FSlateCachedElementsHandle FSlateWindowElementList::PopPaintingWidget(const SWid
 	return FSlateCachedElementsHandle::Invalid;
 }
 
-/*
-int32 FSlateWindowElementList::PushBatchPriortyGroup(const SWidget& CurrentWidget)
-{
-	int32 NewPriorityGroup = 0;
-/ *
-	if (GSlateEnableGlobalInvalidation)
-	{
-		NewPriorityGroup = BatchDepthPriorityStack.Add_GetRef(CurrentWidget.FastPathProxyHandle.IsValid() ? CurrentWidget.FastPathProxyHandle.GetIndex() : 0);
-	}
-	else
-	{
-		NewPriorityGroup = BatchDepthPriorityStack.Add_GetRef(MaxPriorityGroup + 1);
-		//NewPriorityGroup = BatchDepthPriorityStack.Add_GetRef(0);
-	}
-
-	// Should be +1 or the first overlay slot will not appear on top of stuff below it?
-	// const int32 NewPriorityGroup = BatchDepthPriorityStack.Add_GetRef(BatchDepthPriorityStack.Num() ? BatchDepthPriorityStack.Top()+1 : 1);
-
-	MaxPriorityGroup = FMath::Max(NewPriorityGroup, MaxPriorityGroup);* /
-	return NewPriorityGroup;
-}
-
-int32 FSlateWindowElementList::PushAbsoluteBatchPriortyGroup(int32 BatchPriorityGroup)
-{
-	return 0;// return BatchDepthPriorityStack.Add_GetRef(BatchPriorityGroup);
-}
-
-void FSlateWindowElementList::PopBatchPriortyGroup()
-{
-	//BatchDepthPriorityStack.Pop();
-}*/
-
 void FSlateWindowElementList::PushCachedElementData(FSlateCachedElementData& CachedElementData)
 {
 	check(&CachedElementData); 
@@ -378,6 +354,7 @@ void FSlateWindowElementList::ResetElementList()
 
 	DeferredPaintList.Reset();
 
+	FinishMergeRenderBatches();
 	BatchData.ResetData();
 	BatchDataHDR.ResetData();
 
@@ -573,8 +550,6 @@ FSlateRenderBatch& FSlateCachedElementList::AddRenderBatch(int32 InLayer, const
 	CachedRenderBatchIndices.Add(RenderBatchIndex);
 
 	return AddedBatchRef;
-	
-	//return CachedBatches.Emplace_GetRef(InLayer, InShaderParams, InResource, InPrimitiveType, InShaderType, InDrawEffects, InDrawFlags, SceneIndex, &CachedRenderingData->Vertices, &CachedRenderingData->Indices, CachedRenderingData->Vertices.Num(), CachedRenderingData->Indices.Num());
 }
 
 void FSlateCachedElementList::AddCachedClipState(FSlateCachedClipState& ClipStateToCache)
diff --git a/Engine/Source/Runtime/SlateCore/Private/Rendering/ElementBatcher.cpp b/Engine/Source/Runtime/SlateCore/Private/Rendering/ElementBatcher.cpp
--- a/Engine/Source/Runtime/SlateCore/Private/Rendering/ElementBatcher.cpp
+++ b/Engine/Source/Runtime/SlateCore/Private/Rendering/ElementBatcher.cpp
@@ -52,24 +52,6 @@ FSlateElementBatch::FSlateElementBatch(TWeakPtr<ICustomSlateElement, ESPMode::Th
 {
 }
 
-void FSlateElementBatch::SaveClippingState(const TArray<FSlateClippingState>& PrecachedClipStates)
-{
-	/*// Do cached first
-	if (BatchKey.ClipStateHandle.GetCachedClipState().IsSet())
-	{
-		const TSharedPtr<FSlateClippingState>& CachedState = BatchKey.ClipStateHandle.GetCachedClipState().GetValue();
-		if (CachedState.IsValid())
-		{
-			ClippingState = *CachedState;
-		}
-	}
-	else if (PrecachedClipStates.IsValidIndex(BatchKey.ClipStateHandle.GetPrecachedClipIndex()))
-	{
-		// Store the clipping state so we can use it later for rendering.
-		ClippingState = PrecachedClipStates[BatchKey.ClipStateHandle.GetPrecachedClipIndex()];
-	}*/
-}
-
 FSlateBatchData::FSlateBatchData()
 	: FirstRenderBatchIndex(INDEX_NONE)
 	, NumLayers(0)
@@ -108,36 +90,6 @@ FSlateRenderBatch& FSlateBatchData::AddRenderBatch(int32 InLayer, const FShaderP
 	return RenderBatches.Emplace_GetRef(InLayer, InShaderParams, InResource, InPrimitiveType, InShaderType, InDrawEffects, InDrawFlags, SceneIndex, &UncachedSourceBatchVertices, &UncachedSourceBatchIndices, UncachedSourceBatchVertices.Num(), UncachedSourceBatchIndices.Num());
 }
 
-void FSlateBatchData::AddCachedBatches(const TSparseArray<FSlateRenderBatch>& InCachedBatches)
-{
-	RenderBatches.Reserve(RenderBatches.Num() + InCachedBatches.Num());
-
-	for (const FSlateRenderBatch& CachedBatch : InCachedBatches)
-	{
-		RenderBatches.Add(CachedBatch);
-	}
-}
-
-void FSlateBatchData::AddCachedBatchesToBatchData(FSlateBatchData* BatchDataSDR, FSlateBatchData* BatchDataHDR, const TSparseArray<FSlateRenderBatch>& InCachedBatches)
-{
-	TArray<FSlateRenderBatch>& RenderBatchesSDR = BatchDataSDR->RenderBatches;
-	TArray<FSlateRenderBatch>& RenderBatchesHDR = BatchDataHDR->RenderBatches;
-
-	RenderBatchesSDR.Reserve(RenderBatchesSDR.Num() + InCachedBatches.Num());
-	RenderBatchesHDR.Reserve(RenderBatchesHDR.Num() + InCachedBatches.Num());
-	for (const FSlateRenderBatch& CachedBatch : InCachedBatches)
-	{
-		if (EnumHasAnyFlags(CachedBatch.GetDrawFlags(), ESlateBatchDrawFlag::HDR))
-		{
-			RenderBatchesHDR.Add(CachedBatch);
-		}
-		else
-		{
-			RenderBatchesSDR.Add(CachedBatch);
-		}
-	}
-}
-
 void FSlateBatchData::FillBuffersFromNewBatch(FSlateRenderBatch& Batch, FSlateVertexArray& FinalVertices, FSlateIndexArray& FinalIndices)
 {
 	if(Batch.HasVertexData())
@@ -188,7 +140,7 @@ void FSlateBatchData::CombineBatches(FSlateRenderBatch& FirstBatch, FSlateRender
 
 void FSlateBatchData::MergeRenderBatches()
 {
-	SCOPE_CYCLE_COUNTER(STAT_SlateRTCreateBatches);
+	SCOPED_NAMED_EVENT_TEXT("Slate::MergeRenderBatches", FColor::Magenta);
 
 	if(RenderBatches.Num())
 	{
@@ -215,7 +167,6 @@ void FSlateBatchData::MergeRenderBatches()
 			);
 		}
 
-
 		NumBatches = 0;
 		NumLayers = 0;
 
@@ -288,6 +239,9 @@ void FSlateBatchData::MergeRenderBatches()
 #endif
 			PrevBatch = &CurBatch;
 		}
+
+		MaxNumFinalVertices = FMath::Max(MaxNumFinalVertices, FinalVertexData.Num());
+		MaxNumFinalIndices = FMath::Max(MaxNumFinalVertices, FinalIndexData.Num());
 	}
 }
 
@@ -356,6 +310,8 @@ void FSlateElementBatcher::AddElements(FSlateWindowElementList& WindowElementLis
 		}
 	}
 
+	WindowElementList.StartMergeRenderBatches();
+
 	// Done with the element list
 	BatchData = nullptr;
 	BatchDataHDR = nullptr;
@@ -3076,50 +3032,6 @@ void FSlateElementBatcher::AddCustomVerts(const FSlateCustomVertsElement& DrawEl
 		RenderBatch.AddIndices(DrawElement.Indices);
 
 	}
-	/*FElementBatchMap& LayerToElementBatches = CurrentDrawLayer->GetElementBatchMap();
-
-	const FSlateCustomVertsPayload& InPayload = DrawElement.GetDataPayload<FSlateCustomVertsPayload>();
-	uint32 Layer = DrawElement.GetAbsoluteLayer();
-
-	if (InPayload.Vertices.Num() >0)
-	{
-		// See if the layer already exists.
-		TUniqueObj<FElementBatchArray>* ElementBatches = LayerToElementBatches.Find(Layer);
-		if (!ElementBatches)
-		{
-			// The layer doesn't exist so make it now
-			ElementBatches = &LayerToElementBatches.Add( Layer );
-		}
-		check(ElementBatches);
-
-		FSlateElementBatch NewBatch(
-			InPayload.ResourceProxy != nullptr ? InPayload.ResourceProxy->Resource : nullptr,
-			FShaderParams(),
-			ESlateShader::Custom,
-			ESlateDrawPrimitive::TriangleList,
-			DrawElement.GetDrawEffects(),
-			DrawElement.GetBatchFlags(),
-			DrawElement,
-			InPayload.NumInstances,
-			InPayload.InstanceOffset,
-			InPayload.InstanceData
-		);
-
-		NewBatch.SaveClippingState(*PrecachedClippingStates);
-
-		int32 Index = (*ElementBatches)->Add(NewBatch);
-		FSlateElementBatch* ElementBatch = &(**ElementBatches)[Index];
-
-		BatchData->AssignVertexArrayToBatch(*ElementBatch);
-		BatchData->AssignIndexArrayToBatch(*ElementBatch);
-
-		FSlateVertexArray& BatchVertices = BatchData->GetBatchVertexList(*ElementBatch);
-		FSlateIndexArray& BatchIndices = BatchData->GetBatchIndexList(*ElementBatch);
-
-		// Vertex Buffer since  it is already in slate format it is a straight copy
-		BatchVertices = InPayload.Vertices;
-		BatchIndices = InPayload.Indices;
-	}*/
 }
 
 void FSlateElementBatcher::AddPostProcessPass(const FSlatePostProcessElement& DrawElement, FVector2f WindowSize)
@@ -3887,4 +3799,3 @@ void FSlateElementBatcher::ResetBatches()
 	SkipDefaultUpdatePostBuffers = ESlatePostRT::None;
 	NumPostProcessPasses = 0;
 }
-
diff --git a/Engine/Source/Runtime/SlateCore/Public/Rendering/DrawElements.h b/Engine/Source/Runtime/SlateCore/Public/Rendering/DrawElements.h
--- a/Engine/Source/Runtime/SlateCore/Public/Rendering/DrawElements.h
+++ b/Engine/Source/Runtime/SlateCore/Public/Rendering/DrawElements.h
@@ -6,6 +6,7 @@
 #include "Rendering/DrawElementCoreTypes.h"
 #include "Rendering/DrawElementTypes.h"
 #include "Rendering/DrawElementPayloads.h"
+#include "Tasks/Task.h"
 
 #if UE_ENABLE_INCLUDE_ORDER_DEPRECATED_IN_5_2
 #include "Fonts/FontCache.h"
@@ -394,8 +395,17 @@ public:
 	SLATECORE_API void ResetElementList();
 
 
-	FSlateBatchData& GetBatchData() { return BatchData; }
-	FSlateBatchData& GetBatchDataHDR() { return BatchDataHDR; }
+	FSlateBatchData& GetBatchData()
+	{
+		FinishMergeRenderBatches();
+		return BatchData;
+	}
+
+	FSlateBatchData& GetBatchDataHDR()
+	{
+		FinishMergeRenderBatches();
+		return BatchDataHDR;
+	}
 
 	SLATECORE_API void SetRenderTargetWindow(SWindow* InRenderTargetWindow);
 
@@ -437,6 +447,10 @@ private:
 	TArrayView<FSlateCachedElementData* const> GetCachedElementDataList() const { return MakeArrayView(CachedElementDataList.GetData(), CachedElementDataList.Num()); }
 
 	FSlateCachedElementData* GetCurrentCachedElementData() const { return CachedElementDataListStack.Num() ? CachedElementDataList[CachedElementDataListStack.Top()] : nullptr; }
+
+	SLATECORE_API void StartMergeRenderBatches();
+	SLATECORE_API void FinishMergeRenderBatches();
+
 private:
 	/**
 	* Window which owns the widgets that are being painted but not necessarily rendered to
@@ -460,6 +474,8 @@ private:
 	/** Batched data used for rendering */
 	FSlateBatchData BatchDataHDR;
 
+	UE::Tasks::FTask MergeBatchDataTask;
+
 	/**  */
 	FSlateClippingManager ClippingManager;
 
diff --git a/Engine/Source/Runtime/SlateCore/Public/Rendering/ElementBatcher.h b/Engine/Source/Runtime/SlateCore/Public/Rendering/ElementBatcher.h
--- a/Engine/Source/Runtime/SlateCore/Public/Rendering/ElementBatcher.h
+++ b/Engine/Source/Runtime/SlateCore/Public/Rendering/ElementBatcher.h
@@ -46,8 +46,6 @@ public:
 	FSlateElementBatch(const FSlateShaderResource* InShaderResource, const FShaderParams& InShaderParams, ESlateShader ShaderType, ESlateDrawPrimitive PrimitiveType, ESlateDrawEffect InDrawEffects, ESlateBatchDrawFlag InBatchFlags, const FSlateDrawElement& InDrawElement, int32 InstanceCount = 0, uint32 InstanceOffset = 0, ISlateUpdatableInstanceBuffer* InstanceData = nullptr);
 	FSlateElementBatch(TWeakPtr<ICustomSlateElement, ESPMode::ThreadSafe> InCustomDrawer, const FSlateDrawElement& InDrawElement);
 
-	void SaveClippingState(const TArray<FSlateClippingState>& PrecachedClipStates);
-
 	bool operator==(const FSlateElementBatch& Other) const
 	{
 		return BatchKey == Other.BatchKey && ShaderResource == Other.ShaderResource;
@@ -172,6 +170,9 @@ public:
 
 	int32 GetNumFinalBatches() const { return NumBatches; }
 
+	int32 GetMaxNumFinalVertices() const { return MaxNumFinalVertices; }
+	int32 GetMaxNumFinalIndices() const { return MaxNumFinalIndices; }
+
 	const FSlateVertexArray& GetFinalVertexData() const { return FinalVertexData; }
 	const FSlateIndexArray& GetFinalIndexData() const { return FinalIndexData; }
 
@@ -198,12 +199,6 @@ public:
 		ESlateBatchDrawFlag InDrawFlags,
 		int8 SceneIndex);
 
-	/** Adds a cached batch, used in retained rendering */
-	void AddCachedBatches(const TSparseArray<FSlateRenderBatch>& InCachedBatches);
-
-	UE_DEPRECATED(5.5, "AddCachedBatchesToBatchData is no longer needed since this class is a friend FSlateElementBatcher. See: FSlateElementBatcher::AddCachedElements.")
-	static void AddCachedBatchesToBatchData(FSlateBatchData* BatchDataSDR, FSlateBatchData* BatchDataHDR, const TSparseArray<FSlateRenderBatch>& InCachedBatches);
-
 public:
 	friend FSlateElementBatcher;
 
@@ -223,6 +218,9 @@ private:
 	FSlateVertexArray FinalVertexData;
 	FSlateIndexArray  FinalIndexData;
 
+	int32 MaxNumFinalVertices = 0;
+	int32 MaxNumFinalIndices = 0;
+
 	int32 FirstRenderBatchIndex;
 
 	/** */
diff --git a/Engine/Source/Runtime/SlateCore/Public/Rendering/RenderingCommon.h b/Engine/Source/Runtime/SlateCore/Public/Rendering/RenderingCommon.h
--- a/Engine/Source/Runtime/SlateCore/Public/Rendering/RenderingCommon.h
+++ b/Engine/Source/Runtime/SlateCore/Public/Rendering/RenderingCommon.h
@@ -19,6 +19,10 @@
 
 #include "RenderingCommon.generated.h"
 
+class FRHICommandListImmediate;
+class FRDGBuilder;
+class FRDGTexture;
+class FSlateElementBatcher;
 class FSlateInstanceBufferUpdate;
 class FWidgetStyle;
 class SWidget;
@@ -876,75 +880,57 @@ public:
 };
 
 /**
- * An interface for a custom slate drawing element
- * Implementers of this interface are expected to handle destroying this interface properly when a separate 
- * rendering thread may have access to it. (I.E this cannot be destroyed from a different thread if the rendering thread is using it)
+ * An interface for a custom slate elements. Lifetime must support pipelining to the render thread.
  */
 class ICustomSlateElement
 {
 public:
+	virtual ~ICustomSlateElement() {}
+
+	/** Called from the game thread during element batching. */
+	virtual void PostCustomElementAdded(FSlateElementBatcher& ElementBatcher) const {}
+
+	struct FDrawPassInputs
+	{
+		FRDGTexture* OutputTexture = nullptr;
+		FIntRect SceneViewRect;
+		FMatrix44f ElementsMatrix;
+		FVector2f ElementsOffset;
+		EDisplayColorGamut HDRDisplayColorGamut = EDisplayColorGamut::sRGB_D65;
+		ESlatePostRT UsedSlatePostBuffers = ESlatePostRT::None;
+		bool bOutputIsHDRDisplay = false;
+		bool bWireFrame = false;
+	};
+
+	/** Called from the render thread during element drawing. */
+	virtual void Draw_RenderThread(FRDGBuilder& GraphBuilder, const FDrawPassInputs& Inputs) {};
+
+	//////////////////////////////////////////////////////////////////////////
+	// Deprecated API
 
-	/** Struct describing current draw state for the custom drawer */
-	struct FSlateCustomDrawParams
+	struct UE_DEPRECATED(5.5, "Use ICustomSlateElement::Draw instead") FSlateCustomDrawParams
 	{
 		FMatrix44f ViewProjectionMatrix;
 		FVector2f ViewOffset;
 		FIntRect ViewRect;
-		EDisplayColorGamut HDRDisplayColorGamut;
-		ESlatePostRT UsedSlatePostBuffers;
-		bool bWireFrame;
-		bool bIsHDR;
-
-		FSlateCustomDrawParams()
-			: ViewProjectionMatrix(FMatrix44f())
-			, ViewOffset(0.f, 0.f)
-			, ViewRect(FIntRect())
-			, HDRDisplayColorGamut(EDisplayColorGamut::sRGB_D65)
-			, UsedSlatePostBuffers(ESlatePostRT::None)
-			, bWireFrame(false)
-			, bIsHDR(false)
-		{
-		}
+		EDisplayColorGamut HDRDisplayColorGamut = EDisplayColorGamut::sRGB_D65;
+		ESlatePostRT UsedSlatePostBuffers = ESlatePostRT::None;
+		bool bWireFrame = false;
+		bool bIsHDR = false;
 	};
 
-public:
-	virtual ~ICustomSlateElement() {}
-
-	UE_DEPRECATED(5.4, "Please override Draw_RenderThread instead and modify your function signature to accept 'FSlateCustomParams& Params'")
-	virtual void DrawRenderThread(class FRHICommandListImmediate& RHICmdList, const void* RenderTarget) 
-	{
-	}
+	UE_DEPRECATED(5.4, "Use ICustomSlateElement::Draw instead.")
+	virtual void DrawRenderThread(FRHICommandListImmediate& RHICmdList, const void* RenderTarget) {}
 
-	/** 
-	 * Called from the rendering thread when it is time to render the element
-	 *
-	 * @param RenderTarget				handle to the platform specific render target implementation.  Note this is already bound by Slate initially 
-	 * @param Params					Params about current draw state 
-	 * @param RenderingPolicyInterface	Interface to current rendering policy
-	 */
-	virtual void Draw_RenderThread(class FRHICommandListImmediate& RHICmdList, const void* RenderTarget, const FSlateCustomDrawParams& Params)
-	{
+	UE_DEPRECATED(5.5, "Use ICustomSlateElement::Draw instead.")
 PRAGMA_DISABLE_DEPRECATION_WARNINGS
-		DrawRenderThread(RHICmdList, RenderTarget);
+	virtual void Draw_RenderThread(FRHICommandListImmediate& RHICmdList, const void* RenderTarget, const FSlateCustomDrawParams& Params) {}
 PRAGMA_ENABLE_DEPRECATION_WARNINGS
-	}
 
-	/**
-	 * Called from the game thread during element batching
-	 *
-	 * @param ElementBatcher	Elementbatcher that added the custom element
-	 */
-	virtual void PostCustomElementAdded(class FSlateElementBatcher& ElementBatcher) const {}
+	UE_DEPRECATED(5.5, "Use ICustomSlateElement::Draw instead.")
+	virtual bool UsesAdditionalRHIParams() const { return false; }
 
-	/**
-	 * If true will cast to an ICustomSlateElementRHI & call Draw_RenderThread with additional RHI params on that instead.
-	 * 
-	 * Note: While a bool to determine cast is not desirable, it is needed due to RHI module reference constraints
-	 */
-	virtual bool UsesAdditionalRHIParams() const 
-	{
-		return false;
-	}
+	//////////////////////////////////////////////////////////////////////////
 };
 
 /*
diff --git a/Engine/Source/Runtime/SlateCore/Public/Rendering/SlateRendererTypes.h b/Engine/Source/Runtime/SlateCore/Public/Rendering/SlateRendererTypes.h
--- a/Engine/Source/Runtime/SlateCore/Public/Rendering/SlateRendererTypes.h
+++ b/Engine/Source/Runtime/SlateCore/Public/Rendering/SlateRendererTypes.h
@@ -19,6 +19,7 @@ enum class ESlatePostRT : uint8
 	ESlatePostRT_2 = 1 << 2,
 	ESlatePostRT_3 = 1 << 3,
 	ESlatePostRT_4 = 1 << 4,
+	All = ESlatePostRT_0 | ESlatePostRT_1 | ESlatePostRT_2 | ESlatePostRT_3 | ESlatePostRT_4,
 	Num = 5
 };
 
diff --git a/Engine/Source/Runtime/SlateRHIRenderer/Private/FX/SlatePostBufferBlur.cpp b/Engine/Source/Runtime/SlateRHIRenderer/Private/FX/SlatePostBufferBlur.cpp
--- a/Engine/Source/Runtime/SlateRHIRenderer/Private/FX/SlatePostBufferBlur.cpp
+++ b/Engine/Source/Runtime/SlateRHIRenderer/Private/FX/SlatePostBufferBlur.cpp
@@ -11,23 +11,22 @@
 /////////////////////////////////////////////////////
 // FSlatePostBufferBlurProxy
 
-void FSlatePostBufferBlurProxy::PostProcess_Renderthread(FRHICommandListImmediate& RHICmdList, FRHITexture* Src, FRHITexture* Dst, FIntRect SrcRect, FIntRect DstRect, FSlateRHIRenderingPolicyInterface InRenderingPolicy)
+void FSlatePostBufferBlurProxy::PostProcess_Renderthread(FRDGBuilder& GraphBuilder, const FScreenPassTexture& InputTexture, const FScreenPassTexture& OutputTexture)
 {
-	if (InRenderingPolicy.IsValid())
+	RDG_EVENT_SCOPE(GraphBuilder, "SlatePostBufferBlur");
+
+	if (GaussianBlurStrength_RenderThread < UE_SMALL_NUMBER)
 	{
-		if (GaussianBlurStrength_RenderThread < UE_SMALL_NUMBER)
-		{
-			// No real blur, just copy
-			FRHICopyTextureInfo CopyInfo;
-			CopyInfo.SourcePosition = FIntVector(SrcRect.Min.X, SrcRect.Min.Y, 0);
-			CopyInfo.Size = FIntVector(DstRect.Width(), DstRect.Height(), 1);
-			TransitionAndCopyTexture(RHICmdList, Src, Dst, CopyInfo);
-		}
-		else
-		{
-			// Use rendering policy to perform blur post process with desired Src / Dst & respective extents
-			InRenderingPolicy.BlurRectExternal(RHICmdList, Src, Dst, SrcRect, DstRect, GaussianBlurStrength_RenderThread);
-		}
+		AddDrawTexturePass(GraphBuilder, FScreenPassViewInfo(), InputTexture, FScreenPassRenderTarget(OutputTexture));
+	}
+	else
+	{
+		FSlatePostProcessSimpleBlurPassInputs BlurInputs;
+		BlurInputs.InputTexture = InputTexture;
+		BlurInputs.OutputTexture = OutputTexture;
+		BlurInputs.Strength = GaussianBlurStrength_RenderThread;
+
+		AddSlatePostProcessBlurPass(GraphBuilder, BlurInputs);
 	}
 }
 
@@ -90,49 +89,6 @@ USlatePostBufferBlur::~USlatePostBufferBlur()
 	RenderThreadProxy = nullptr;
 }
 
-void USlatePostBufferBlur::PostProcess(FRenderResource* InViewInfo, FRenderResource* InViewportTexture, FVector2D InElementWindowSize, FSlateRHIRenderingPolicyInterface InRenderingPolicy, UTextureRenderTarget2D* InSlatePostBuffer)
-{
-	if (!InRenderingPolicy.IsValid())
-	{
-		return;
-	};
-
-	// Explicit param copy to avoid renderthread from reading value during gamethread write
-	float GaussianBlurStrengthCopy = GaussianBlurStrength;
-
-	// Enqueue default post process command, will trigger on scene before any other slate element draws
-	ENQUEUE_RENDER_COMMAND(FUpdateSlatePostBuffersWithFX_Blur)([InViewInfo, InViewportTexture, InElementWindowSize, InRenderingPolicy, InSlatePostBuffer, GaussianBlurStrengthCopy](FRHICommandListImmediate& RHICmdList)
-	{
-		// Get Backbuffer, which can vary between PIE or standalone
-		FTextureRHIRef BackBuffer = USlateRHIPostBufferProcessor::GetBackbuffer_RenderThread(InViewInfo, InViewportTexture, InElementWindowSize, RHICmdList);
-
-		if (BackBuffer)
-		{
-			// Get Src / Dst textures & their rects, again may vary between PIE or standalone
-			// Here we can simply use Src Rect since the Src texture in PIE is the 'BufferedRT' scene backbuffer without the editor
-			FTextureRHIRef Src = USlateRHIPostBufferProcessor::GetSrcTexture_RenderThread(BackBuffer, InViewportTexture);
-			FTextureReferenceRHIRef& Dst = USlateRHIPostBufferProcessor::GetDstTexture_RenderThread(InSlatePostBuffer);
-			FIntPoint DstExtent = USlateRHIPostBufferProcessor::GetDstExtent_RenderThread(BackBuffer, InViewportTexture);
-			FIntRect SrcRect = FIntRect(0, 0, Src->GetSizeX(), Src->GetSizeY());
-			FIntRect DstRect = FIntRect(0, 0, DstExtent.X, DstExtent.Y);
-
-			if (GaussianBlurStrengthCopy < UE_SMALL_NUMBER)
-			{
-				// No real blur, just copy
-				FRHICopyTextureInfo CopyInfo;
-				CopyInfo.SourcePosition = FIntVector(SrcRect.Min.X, SrcRect.Min.Y, 0);
-				CopyInfo.Size = FIntVector(DstRect.Width(), DstRect.Height(), 1);
-				TransitionAndCopyTexture(RHICmdList, Src, Dst, CopyInfo);
-			}
-			else
-			{
-				// Use rendering policy to perform blur post process with desired Src / Dst & respective extents
-				InRenderingPolicy.BlurRectExternal(RHICmdList, Src, Dst, SrcRect, DstRect, GaussianBlurStrengthCopy);
-			}
-		}
-	});
-}
-
 TSharedPtr<FSlateRHIPostBufferProcessorProxy> USlatePostBufferBlur::GetRenderThreadProxy()
 {
 	if (!RenderThreadProxy && IsInGameThread())
@@ -142,4 +98,4 @@ TSharedPtr<FSlateRHIPostBufferProcessorProxy> USlatePostBufferBlur::GetRenderThr
 		RenderThreadProxy->SetOwningProcessorObject(this);
 	}
 	return RenderThreadProxy;
-}
+}
\ No newline at end of file
diff --git a/Engine/Source/Runtime/SlateRHIRenderer/Private/FX/SlateRHIPostBufferProcessor.cpp b/Engine/Source/Runtime/SlateRHIRenderer/Private/FX/SlateRHIPostBufferProcessor.cpp
--- a/Engine/Source/Runtime/SlateRHIRenderer/Private/FX/SlateRHIPostBufferProcessor.cpp
+++ b/Engine/Source/Runtime/SlateRHIRenderer/Private/FX/SlateRHIPostBufferProcessor.cpp
@@ -7,39 +7,4 @@
 #include "StereoRendering.h"
 #include "UnrealEngine.h"
 
-#include UE_INLINE_GENERATED_CPP_BY_NAME(SlateRHIPostBufferProcessor)
-
-FTextureRHIRef USlateRHIPostBufferProcessor::GetBackbuffer_RenderThread(FRenderResource* InViewInfo, FRenderResource* InViewportTexture, FVector2D InElementWindowSize, FRHICommandListImmediate& InRHICmdList)
-{
-	FViewportInfo* ViewInfo = static_cast<FViewportInfo*>(InViewInfo);
-	FSlateRenderTargetRHI* ViewportTexture = static_cast<FSlateRenderTargetRHI*>(InViewportTexture);
-
-	bool bRenderedStereo = false;
-	if (FSlateRHIRenderer::GetDrawToVRRenderTarget() == 0 && GEngine && IsValidRef(ViewInfo->GetRenderTargetTexture()) && GEngine->StereoRenderingDevice.IsValid())
-	{
-		GEngine->StereoRenderingDevice->RenderTexture_RenderThread(InRHICmdList, RHIGetViewportBackBuffer(ViewInfo->ViewportRHI), ViewInfo->GetRenderTargetTexture(), InElementWindowSize);
-		bRenderedStereo = true;
-	}
-
-	FTextureRHIRef ViewportRT = bRenderedStereo ? nullptr : ViewInfo->GetRenderTargetTexture();
-	FTextureRHIRef BackBuffer = (ViewportRT) ? ViewportRT : RHIGetViewportBackBuffer(ViewInfo->ViewportRHI);
-
-	return BackBuffer;
-}
-
-FTextureRHIRef USlateRHIPostBufferProcessor::GetSrcTexture_RenderThread(FTextureRHIRef InBackBuffer, FRenderResource* InViewportTexture)
-{
-	FSlateRenderTargetRHI* ViewportTexture = static_cast<FSlateRenderTargetRHI*>(InViewportTexture);
-	return GIsEditor ? ViewportTexture->GetRHIRef() : InBackBuffer;
-}
-
-FTextureReferenceRHIRef& USlateRHIPostBufferProcessor::GetDstTexture_RenderThread(UTextureRenderTarget2D* InSlatePostBuffer)
-{
-	return InSlatePostBuffer->TextureReference.TextureReferenceRHI;
-}
-
-FIntPoint USlateRHIPostBufferProcessor::GetDstExtent_RenderThread(FTextureRHIRef InBackBuffer, FRenderResource* InViewportTexture)
-{
-	FSlateRenderTargetRHI* ViewportTexture = static_cast<FSlateRenderTargetRHI*>(InViewportTexture);
-	return GIsEditor ? FIntPoint(ViewportTexture->GetWidth(), ViewportTexture->GetHeight()) : InBackBuffer->GetDesc().Extent;
-}
\ No newline at end of file
+#include UE_INLINE_GENERATED_CPP_BY_NAME(SlateRHIPostBufferProcessor)
\ No newline at end of file
diff --git a/Engine/Source/Runtime/SlateRHIRenderer/Private/Slate3DRenderer.cpp b/Engine/Source/Runtime/SlateRHIRenderer/Private/Slate3DRenderer.cpp
--- a/Engine/Source/Runtime/SlateRHIRenderer/Private/Slate3DRenderer.cpp
+++ b/Engine/Source/Runtime/SlateRHIRenderer/Private/Slate3DRenderer.cpp
@@ -8,54 +8,34 @@
 #include "SlateRHIRenderer.h"
 #include "Rendering/ElementBatcher.h"
 #include "Types/SlateVector2.h"
+#include "RenderGraphUtils.h"
+#include "SlatePostProcessor.h"
 
 DECLARE_GPU_STAT_NAMED(Slate3D, TEXT("Slate 3D"));
 
-FSlate3DRenderer::FSlate3DRenderer( TSharedRef<FSlateFontServices> InSlateFontServices, TSharedRef<FSlateRHIResourceManager> InResourceManager, bool bUseGammaCorrection )
+FSlate3DRenderer::FSlate3DRenderer( TSharedRef<FSlateFontServices> InSlateFontServices, TSharedRef<FSlateRHIResourceManager> InResourceManager, bool bInUseGammaCorrection )
 	: SlateFontServices( InSlateFontServices )
 	, ResourceManager( InResourceManager )
-	, bRenderTargetWasCleared(false)
+	, bGammaCorrection(bInUseGammaCorrection)
 {
-	const int32 InitialBufferSize = 200;
-	RenderTargetPolicy = MakeShareable( new FSlateRHIRenderingPolicy( SlateFontServices, ResourceManager, InitialBufferSize ) );
-	RenderTargetPolicy->SetUseGammaCorrection( bUseGammaCorrection );
+	RenderTargetPolicy = MakeShareable( new FSlateRHIRenderingPolicy( SlateFontServices, ResourceManager ) );
 
 	ElementBatcher = MakeUnique<FSlateElementBatcher>(RenderTargetPolicy.ToSharedRef());
 }
 
 void FSlate3DRenderer::Cleanup()
 {
-	if ( RenderTargetPolicy.IsValid() )
-	{
-		RenderTargetPolicy->ReleaseResources();
-	}
-
-	if (IsInGameThread())
-	{
-		// Enqueue a command to unlock the draw buffer after all windows have been drawn
-		ENQUEUE_RENDER_COMMAND(FSlate3DRenderer_Cleanup)(
-			[this](FRHICommandListImmediate& RHICmdList)
-			{
-				DepthStencil.SafeRelease();
-			}
-		);
-	}
-	else
-	{
-		DepthStencil.SafeRelease();
-	}
-
 	BeginCleanup(this);
 }
 
-void FSlate3DRenderer::SetUseGammaCorrection(bool bUseGammaCorrection)
+void FSlate3DRenderer::SetUseGammaCorrection(bool bInUseGammaCorrection)
 {
-	RenderTargetPolicy->SetUseGammaCorrection(bUseGammaCorrection);
+	bGammaCorrection = bInUseGammaCorrection;
 }
 
-void FSlate3DRenderer::SetApplyColorDeficiencyCorrection(bool bApplyColorCorrection)
+void FSlate3DRenderer::SetApplyColorDeficiencyCorrection(bool bInAllowColorDeficiencyCorrection)
 {
-	RenderTargetPolicy->SetApplyColorDeficiencyCorrection(bApplyColorCorrection);
+	bAllowColorDeficiencyCorrection = bInAllowColorDeficiencyCorrection;
 }
 
 FSlateDrawBuffer& FSlate3DRenderer::AcquireDrawBuffer()
@@ -93,13 +73,11 @@ void FSlate3DRenderer::ReleaseDrawBuffer(FSlateDrawBuffer& InWindowDrawBuffer)
 	ensureMsgf(bFound, TEXT("It release a DrawBuffer that is not a member of the Slate3DRenderer"));
 #endif
 
-	FSlateDrawBuffer* DrawBuffer = &InWindowDrawBuffer;
 	ENQUEUE_RENDER_COMMAND(SlateReleaseDrawBufferCommand)(
-		[DrawBuffer](FRHICommandListImmediate& RHICmdList)
-		{
-			FSlateReleaseDrawBufferCommand::ReleaseDrawBuffer(RHICmdList, DrawBuffer);
-		}
-	);
+		[DrawBuffer = &InWindowDrawBuffer](FRHICommandListImmediate& RHICmdList)
+	{
+		DrawBuffer->Unlock();
+	});
 }
 
 void FSlate3DRenderer::DrawWindow_GameThread(FSlateDrawBuffer& DrawBuffer)
@@ -134,112 +112,80 @@ void FSlate3DRenderer::DrawWindow_GameThread(FSlateDrawBuffer& DrawBuffer)
 	}
 }
 
-void FSlate3DRenderer::DrawWindowToTarget_RenderThread(FRHICommandListImmediate& InRHICmdList, const FRenderThreadUpdateContext& Context)
+void FSlate3DRenderer::DrawWindowToTarget_RenderThread(FRDGBuilder& GraphBuilder, const FRenderThreadUpdateContext& Context)
 {
-	check(IsInRenderingThread());
 	QUICK_SCOPE_CYCLE_COUNTER(Stat_Slate_WidgetRendererRenderThread);
-	SCOPED_DRAW_EVENT( InRHICmdList, SlateRenderToTarget );
-	SCOPED_GPU_STAT(InRHICmdList, Slate3D);
-
-	checkSlow(Context.RenderTarget);
+	RDG_EVENT_SCOPE(GraphBuilder, "SlateRenderToTarget");
+	RDG_GPU_STAT_SCOPE(GraphBuilder, Slate3D);
+	check(Context.RenderTarget);
 
 	const TArray<TSharedRef<FSlateWindowElementList>>& WindowsToDraw = Context.WindowDrawBuffer->GetWindowElementLists();
 
-	// Enqueue a command to unlock the draw buffer after all windows have been drawn
-	RenderTargetPolicy->BeginDrawingWindows();
-
-	// Set render target and clear.
-	FRHITexture* RTTextureRHI = Context.RenderTarget->GetRenderTargetTexture();
-	InRHICmdList.Transition(FRHITransitionInfo(RTTextureRHI, ERHIAccess::Unknown, ERHIAccess::RTV));
-	
-	FRHIRenderPassInfo RPInfo(RTTextureRHI, ERenderTargetActions::Load_Store);
-	if (Context.bClearTarget)
-	{
-		RPInfo.ColorRenderTargets[0].Action = ERenderTargetActions::Clear_Store;
-	}
-
-	for (int32 WindowIndex = 0; WindowIndex < WindowsToDraw.Num(); WindowIndex++)
-	{
-		FSlateWindowElementList& WindowElementList = *WindowsToDraw[WindowIndex];
+	FRDGTexture* SlateElementsTexture = RegisterExternalTexture(GraphBuilder, Context.RenderTarget->GetRenderTargetTexture(), TEXT("SlateElementsTexture"));
+	const FIntPoint SlateElementsExtent = SlateElementsTexture->Desc.Extent;
 
-		FSlateBatchData& BatchData = WindowElementList.GetBatchData();
+	FRDGTexture* SlateStencilTexture = nullptr;
+	bool bStencilClippingRequired = false;
 
-		if (BatchData.GetRenderBatches().Num() > 0)
-		{
-			RenderTargetPolicy->BuildRenderingBuffers(InRHICmdList, BatchData);
-		
-			FVector2D DrawOffset = Context.WindowDrawBuffer->ViewOffset;
+	TArray<FSlateElementsBuffers, FRDGArrayAllocator> SlateElementsBuffers;
+	SlateElementsBuffers.Reserve(WindowsToDraw.Num());
 
-			FMatrix ProjectionMatrix = FSlateRHIRenderer::CreateProjectionMatrix(RTTextureRHI->GetSizeX(), RTTextureRHI->GetSizeY());
-			FMatrix ViewOffset = FTranslationMatrix::Make(FVector(DrawOffset, 0.0));
-			ProjectionMatrix = ViewOffset * ProjectionMatrix;
+	for (const TSharedRef<FSlateWindowElementList>& WindowElementList : WindowsToDraw)
+	{
+		FSlateBatchData& BatchData = WindowElementList->GetBatchData();
 
-			FSlateBackBuffer BackBufferTarget(Context.RenderTarget->GetRenderTargetTexture(), FIntPoint(RTTextureRHI->GetSizeX(), RTTextureRHI->GetSizeY()));
+		SlateElementsBuffers.Emplace(BuildSlateElementsBuffers(GraphBuilder, BatchData));
 
-			FSlateRenderingParams DrawOptions(ProjectionMatrix, FGameTime::CreateDilated(Context.RealTimeSeconds, Context.DeltaRealTimeSeconds, Context.WorldTimeSeconds, Context.DeltaTimeSeconds));
-			// The scene renderer will handle it in this case
-			DrawOptions.ViewOffset = UE::Slate::CastToVector2f(DrawOffset);
+		bStencilClippingRequired |= BatchData.IsStencilClippingRequired();
+	}
 
-			FTextureRHIRef ColorTarget = Context.RenderTarget->GetRenderTargetTexture();
+	if (bStencilClippingRequired)
+	{
+		SlateStencilTexture = GraphBuilder.CreateTexture(
+			FRDGTextureDesc::Create2D(SlateElementsExtent, PF_DepthStencil, FClearValueBinding::DepthZero, GetSlateTransientDepthStencilFlags()),
+			TEXT("SlateStencilTexture"));
+	}
 
-			if (BatchData.IsStencilClippingRequired())
-			{
-				if (!DepthStencil.IsValid() || ColorTarget->GetSizeXY() != DepthStencil->GetSizeXY())
-				{
-					DepthStencil.SafeRelease();
+	ERenderTargetLoadAction ElementsLoadAction = Context.bClearTarget ? ERenderTargetLoadAction::EClear : ERenderTargetLoadAction::ELoad;
+	
+	const auto ConsumeLoadAction = [] (ERenderTargetLoadAction& InOutLoadAction)
+	{
+		ERenderTargetLoadAction LoadAction = InOutLoadAction;
+		InOutLoadAction = ERenderTargetLoadAction::ELoad;
+		return LoadAction;
+	};
 
-					const FRHITextureCreateDesc Desc =
-						FRHITextureCreateDesc::Create2D(TEXT("SlateWindowDepthStencil"))
-						.SetExtent(ColorTarget->GetSizeXY())
-						.SetFormat(PF_DepthStencil)
-						.SetClearValue(FClearValueBinding::DepthZero)
-						.SetFlags(ETextureCreateFlags::DepthStencilTargetable | ETextureCreateFlags::ShaderResource);
+	for (int32 WindowElementIndex = 0; WindowElementIndex < WindowsToDraw.Num(); ++WindowElementIndex)
+	{
+		const TSharedRef<FSlateWindowElementList>& WindowElementList = WindowsToDraw[WindowElementIndex];
 
-					DepthStencil = RHICreateTexture(Desc);
+		FSlateBatchData& BatchData = WindowElementList->GetBatchData();
 
-					check(IsValidRef(DepthStencil));
-				}
-			}
+		if (BatchData.GetRenderBatches().IsEmpty())
+		{
+			continue;
+		}
 
-			// Ideally we'd have a single render pass for all the windows, but this code reuses a single vertex buffer for each draw, which it updates above in BuildRenderingBuffers(),
-			// and we can't upload data during a render pass. We need to rewrite this to upload all the data to the buffer first and use offsets for each draw.
-			InRHICmdList.BeginRenderPass(RPInfo, TEXT("Slate3D"));
-
-			RenderTargetPolicy->DrawElements(
-				InRHICmdList,
-				BackBufferTarget,
-				ColorTarget,
-				ColorTarget,
-				DepthStencil,
-				BatchData.GetFirstRenderBatchIndex(),
-				BatchData.GetRenderBatches(),
-				DrawOptions
-			);
-
-			// FSlateRHIRenderingPolicy::DrawElements can close the render pass on its own sometimes, so don't do it again.
-			if (InRHICmdList.IsInsideRenderPass())
-			{
-				InRHICmdList.EndRenderPass();
-			}
+		const FVector ElementsOffset(Context.WindowDrawBuffer->ViewOffset, 0.0f);
+		const FMatrix ElementsMatrix(FTranslationMatrix::Make(ElementsOffset) * CreateSlateProjectionMatrix(SlateElementsExtent.X, SlateElementsExtent.Y));
 
-			// each time we do draw content, we reset the flag
-			bRenderTargetWasCleared = false;
-		}
-		// If we have no render command and it's the first clear, we call Begin End to force the clear of the render target. Otherwise, we end up not updating the buffer when all Widget are invisible.
-		else if(!bRenderTargetWasCleared)
+		FSlateDrawElementsPassInputs DrawElementsInputs =
 		{
-			InRHICmdList.BeginRenderPass(RPInfo, TEXT("Slate3D")); 
-			if (InRHICmdList.IsInsideRenderPass())
-			{
-				InRHICmdList.EndRenderPass();
-			}
-			bRenderTargetWasCleared = true;
-		}
+			  .StencilTexture        = SlateStencilTexture
+			, .ElementsTexture       = SlateElementsTexture
+			, .ElementsLoadAction    = ConsumeLoadAction(ElementsLoadAction)
+			, .ElementsBuffers       = SlateElementsBuffers[WindowElementIndex]
+			, .ElementsMatrix        = FMatrix44f(ElementsMatrix)
+			, .ElementsOffset        = FVector2f(ElementsOffset.X, ElementsOffset.Y)
+			, .Time                  = FGameTime::CreateDilated(Context.RealTimeSeconds, Context.DeltaRealTimeSeconds, Context.WorldTimeSeconds, Context.DeltaTimeSeconds)
+			, .bAllowGammaCorrection = bGammaCorrection
+		};
+
+		AddSlateDrawElementsPass(GraphBuilder, *RenderTargetPolicy, DrawElementsInputs, BatchData.GetRenderBatches(), BatchData.GetFirstRenderBatchIndex());
 	}
 
-	FSlateEndDrawingWindowsCommand::EndDrawingWindows(InRHICmdList, Context.WindowDrawBuffer, *RenderTargetPolicy);
-	InRHICmdList.Transition(FRHITransitionInfo(RTTextureRHI, ERHIAccess::RTV, ERHIAccess::SRVMask));
-
-	// Enqueue a command to keep "this" alive.
-	InRHICmdList.EnqueueLambda([Self = SharedThis(this)](FRHICommandListImmediate&){});
-}
+	if (ConsumeLoadAction(ElementsLoadAction) == ERenderTargetLoadAction::EClear)
+	{
+		AddClearRenderTargetPass(GraphBuilder, SlateElementsTexture);
+	}
+}
\ No newline at end of file
diff --git a/Engine/Source/Runtime/SlateRHIRenderer/Private/Slate3DRenderer.h b/Engine/Source/Runtime/SlateRHIRenderer/Private/Slate3DRenderer.h
--- a/Engine/Source/Runtime/SlateRHIRenderer/Private/Slate3DRenderer.h
+++ b/Engine/Source/Runtime/SlateRHIRenderer/Private/Slate3DRenderer.h
@@ -26,7 +26,7 @@ public:
 	virtual FSlateDrawBuffer& AcquireDrawBuffer() override;
 	virtual void ReleaseDrawBuffer(FSlateDrawBuffer& DrawBuffer) override;
 	virtual void DrawWindow_GameThread(FSlateDrawBuffer& DrawBuffer) override;
-	virtual void DrawWindowToTarget_RenderThread(FRHICommandListImmediate& RHICmdList, const struct FRenderThreadUpdateContext& Context) override;
+	virtual void DrawWindowToTarget_RenderThread(FRDGBuilder& GraphBuilder, const struct FRenderThreadUpdateContext& Context) override;
 
 private:
 
@@ -49,9 +49,9 @@ private:
 	/** The draw buffer that is currently free for use by the game thread */
 	uint8 FreeBufferIndex;
 
-	/** The depth buffer texture if any */
-	FTextureRHIRef DepthStencil;
-
 	/** Set to true when the render target was cleared and prevent another clear to be called for nothing.**/
-	bool bRenderTargetWasCleared;
+	bool bRenderTargetWasCleared = false;
+
+	bool bAllowColorDeficiencyCorrection = true;
+	bool bGammaCorrection;
 };
diff --git a/Engine/Source/Runtime/SlateRHIRenderer/Private/SlateElementVertexBuffer.cpp b/Engine/Source/Runtime/SlateRHIRenderer/Private/SlateElementVertexBuffer.cpp
--- a/Engine/Source/Runtime/SlateRHIRenderer/Private/SlateElementVertexBuffer.cpp
+++ b/Engine/Source/Runtime/SlateRHIRenderer/Private/SlateElementVertexBuffer.cpp
@@ -3,3 +3,5 @@
 #include "SlateElementVertexBuffer.h"
 
 DEFINE_STAT(STAT_SlateVertexBufferMemory);
+
+TGlobalResource<FSlateStencilClipVertexBuffer> GSlateStencilClipVertexBuffer;
\ No newline at end of file
diff --git a/Engine/Source/Runtime/SlateRHIRenderer/Private/SlateElementVertexBuffer.h b/Engine/Source/Runtime/SlateRHIRenderer/Private/SlateElementVertexBuffer.h
--- a/Engine/Source/Runtime/SlateRHIRenderer/Private/SlateElementVertexBuffer.h
+++ b/Engine/Source/Runtime/SlateRHIRenderer/Private/SlateElementVertexBuffer.h
@@ -191,3 +191,5 @@ public:
 	/** Returns a friendly name for this buffer. */
 	virtual FString GetFriendlyName() const { return TEXT("SlateElementVertices"); }
 };
+
+extern TGlobalResource<FSlateStencilClipVertexBuffer> GSlateStencilClipVertexBuffer;
\ No newline at end of file
diff --git a/Engine/Source/Runtime/SlateRHIRenderer/Private/SlateMaterialShader.cpp b/Engine/Source/Runtime/SlateRHIRenderer/Private/SlateMaterialShader.cpp
--- a/Engine/Source/Runtime/SlateRHIRenderer/Private/SlateMaterialShader.cpp
+++ b/Engine/Source/Runtime/SlateRHIRenderer/Private/SlateMaterialShader.cpp
@@ -36,10 +36,17 @@ void FSlateMaterialShaderVS::SetViewProjection(FRHIBatchedShaderParameters& Batc
 	SetShaderValue(BatchedParameters, ViewProjection, InViewProjection );
 }
 
-void FSlateMaterialShaderVS::SetMaterialShaderParameters(FRHIBatchedShaderParameters& BatchedParameters, const FSceneView& View, const FMaterialRenderProxy* MaterialRenderProxy, const FMaterial* Material)
+void FSlateMaterialShaderVS::SetMaterialShaderParameters(
+	FRHIBatchedShaderParameters& BatchedParameters,
+	const FSceneInterface* Scene,
+	const TUniformBufferRef<FViewUniformShaderParameters>& ViewUniformBuffer,
+	const FMaterialRenderProxy* MaterialRenderProxy,
+	const FMaterial* Material)
 {
-	SetViewParameters(BatchedParameters, View, View.ViewUniformBuffer);
-	FMaterialShader::SetParameters(BatchedParameters, MaterialRenderProxy, *Material, View);
+	const auto& ViewUniformBufferParameter = GetUniformBufferParameter<FViewUniformShaderParameters>();
+	SetUniformBufferParameter(BatchedParameters, ViewUniformBufferParameter, ViewUniformBuffer);
+
+	FMaterialShader::SetParameters(BatchedParameters, MaterialRenderProxy, *Material, Scene);
 }
 
 bool FSlateMaterialShaderPS::ShouldCompilePermutation(const FMaterialShaderPermutationParameters& Parameters)
@@ -104,13 +111,21 @@ void FSlateMaterialShaderPS::SetBlendState(FGraphicsPipelineStateInitializer& Gr
 	};
 }
 
-void FSlateMaterialShaderPS::SetParameters(FRHIBatchedShaderParameters& BatchedParameters, const FSceneView& View, const FMaterialRenderProxy* MaterialRenderProxy, const FMaterial* Material, const FShaderParams& InShaderParams)
+void FSlateMaterialShaderPS::SetParameters(
+	FRHIBatchedShaderParameters& BatchedParameters,
+	const FSceneInterface* Scene,
+	const TUniformBufferRef<FViewUniformShaderParameters>& ViewUniformBuffer,
+	const FMaterialRenderProxy* MaterialRenderProxy,
+	const FMaterial* Material,
+	const FShaderParams& InShaderParams)
 {
 	SetShaderValue(BatchedParameters, ShaderParams, (FVector4f)InShaderParams.PixelParams);
 	SetShaderValue(BatchedParameters, ShaderParams2, (FVector4f)InShaderParams.PixelParams2);
 
-	SetViewParameters(BatchedParameters, View, View.ViewUniformBuffer);
-	FMaterialShader::SetParameters(BatchedParameters, MaterialRenderProxy, *Material, View);
+	const auto& ViewUniformBufferParameter = GetUniformBufferParameter<FViewUniformShaderParameters>();
+	SetUniformBufferParameter(BatchedParameters, ViewUniformBufferParameter, ViewUniformBuffer);
+
+	FMaterialShader::SetParameters(BatchedParameters, MaterialRenderProxy, *Material, Scene);
 }
 
 void FSlateMaterialShaderPS::SetAdditionalTexture(FRHIBatchedShaderParameters& BatchedParameters, FRHITexture* InTexture, const FSamplerStateRHIRef SamplerState )
diff --git a/Engine/Source/Runtime/SlateRHIRenderer/Private/SlateMaterialShader.h b/Engine/Source/Runtime/SlateRHIRenderer/Private/SlateMaterialShader.h
--- a/Engine/Source/Runtime/SlateRHIRenderer/Private/SlateMaterialShader.h
+++ b/Engine/Source/Runtime/SlateRHIRenderer/Private/SlateMaterialShader.h
@@ -25,16 +25,18 @@ public:
 	 *
 	 * @param InViewProjection	The ViewProjection matrix to use when this shader is bound 
 	 */
-	void SetViewProjection(FRHIBatchedShaderParameters& BatchedParameters, const FMatrix44f& InViewProjection );
+	void SetViewProjection(FRHIBatchedShaderParameters& BatchedParameters, const FMatrix44f& InViewProjection);
 
-	void SetMaterialShaderParameters(FRHIBatchedShaderParameters& BatchedParameters, const FSceneView& View, const FMaterialRenderProxy* MaterialRenderProxy, const FMaterial* Material);
+	void SetMaterialShaderParameters(
+		FRHIBatchedShaderParameters& BatchedParameters,
+		const FSceneInterface* Scene,
+		const TUniformBufferRef<FViewUniformShaderParameters>& ViewUniformBuffer,
+		const FMaterialRenderProxy* MaterialRenderProxy,
+		const FMaterial* Material);
 
-	/** Serializes the shader data */
-	//virtual bool Serialize( FArchive& Ar ) override;
 private:
-	
-		/** ViewProjection parameter used by the shader */
-		LAYOUT_FIELD(FShaderParameter, ViewProjection)
+	/** ViewProjection parameter used by the shader */
+	LAYOUT_FIELD(FShaderParameter, ViewProjection)
 };
 
 class FSlateMaterialShaderPS : public FMaterialShader
@@ -53,7 +55,13 @@ public:
 
 	void SetBlendState(FGraphicsPipelineStateInitializer& GraphicsPSOInit, const FMaterial* Material);
 
-	void SetParameters(FRHIBatchedShaderParameters& BatchedParameters, const FSceneView& View, const FMaterialRenderProxy* MaterialRenderProxy, const FMaterial* Material, const FShaderParams& InShaderParams);
+	void SetParameters(
+		FRHIBatchedShaderParameters& BatchedParameters,
+		const FSceneInterface* Scene,
+		const TUniformBufferRef<FViewUniformShaderParameters>& ViewUniformBuffer,
+		const FMaterialRenderProxy* MaterialRenderProxy,
+		const FMaterial* Material,
+		const FShaderParams& InShaderParams);
 
 	void SetDisplayGammaAndContrast(FRHIBatchedShaderParameters& BatchedParameters, float InDisplayGamma, float InContrast);
 
