Clean-up of Slate RHI renderer and conversion to RDG.

diff --git a/Engine/Source/Runtime/SlateRHIRenderer/Private/SlatePostProcessor.cpp b/Engine/Source/Runtime/SlateRHIRenderer/Private/SlatePostProcessor.cpp
--- a/Engine/Source/Runtime/SlateRHIRenderer/Private/SlatePostProcessor.cpp
+++ b/Engine/Source/Runtime/SlateRHIRenderer/Private/SlatePostProcessor.cpp
@@ -1,870 +1,528 @@
 // Copyright Epic Games, Inc. All Rights Reserved.
 
 #include "SlatePostProcessor.h"
-#include "SlatePostProcessResource.h"
 #include "SlateShaders.h"
-#include "ScreenRendering.h"
-#include "SceneUtils.h"
-#include "RendererInterface.h"
-#include "DataDrivenShaderPlatformInfo.h"
-#include "StaticBoundShaderState.h"
-#include "PipelineStateCache.h"
-#include "CommonRenderResources.h"
-#include "HDRHelper.h"
 #include "RendererUtils.h"
+#include "Interfaces/SlateRHIRenderingPolicyInterface.h"
 
-DECLARE_CYCLE_STAT(TEXT("Slate PostProcessing RT"), STAT_SlatePostProcessingRTTime, STATGROUP_Slate);
-DECLARE_CYCLE_STAT(TEXT("Slate ColorDeficiency RT"), STAT_SlateColorDeficiencyRTTime, STATGROUP_Slate);
+//////////////////////////////////////////////////////////////////////////
 
-int GSlateEnableDeleteUnusedPostProcess = 0;
-static FAutoConsoleVariableRef CVarSlateEnableDeleteUnusedPostProcess(
-	TEXT("Slate.EnableDeleteUnusedPostProcess"),
-	GSlateEnableDeleteUnusedPostProcess,
-	TEXT("Greater than zero implies that post process render targets will be deleted when they are not used after n frames.")
-);
-
-static const int32 NumIntermediateTargets = 2;
-
-FSlatePostProcessResource* FindSlatePostProcessResource(TArray<FSlatePostProcessResource*>& IntermediateTargetsArray, EPixelFormat PixelFormat)
+float GetSlateHDRUILevel()
 {
-	for (int32 ui = 0; ui < IntermediateTargetsArray.Num(); ++ui)
-	{
-		if (IntermediateTargetsArray[ui]->GetPixelFormat() == PixelFormat || IntermediateTargetsArray[ui]->GetPixelFormat() == PF_Unknown)
-		{
-			return IntermediateTargetsArray[ui];
-		}
-	}
+	static auto CVar = IConsoleManager::Get().FindConsoleVariable(TEXT("r.HDR.UI.Level"));
+	return CVar ? CVar->GetFloat() : 1.0f;
+}
 
-	FSlatePostProcessResource* NewSlatePostProcessResource = new FSlatePostProcessResource(NumIntermediateTargets);
-	IntermediateTargetsArray.Add(NewSlatePostProcessResource);
-	BeginInitResource(NewSlatePostProcessResource);
-	return NewSlatePostProcessResource;
+float GetSlateHDRUILuminance()
+{
+	static auto CVar = IConsoleManager::Get().FindConsoleVariable(TEXT("r.HDR.UI.Luminance"));
+	return CVar ? CVar->GetFloat() : 300.0f;
 }
 
-FSlatePostProcessor::FSlatePostProcessor()
+ETextureCreateFlags GetSlateTransientRenderTargetFlags()
 {
-	uint32 MaximumDifferentPixelFormat = 8;
- 	for (uint32 ui=0; ui < MaximumDifferentPixelFormat; ++ui)
- 	{
- 		IntermediateTargetsArray.Add(new FSlatePostProcessResource(NumIntermediateTargets));
- 		BeginInitResource(IntermediateTargetsArray[ui]);
- 	}
+	return ETextureCreateFlags::RenderTargetable | ETextureCreateFlags::ShaderResource | ETextureCreateFlags::FastVRAM
+		// Avoid fast clear metadata when this flag is set, since we'd otherwise have to clear transient render targets instead of discard.
+#if PLATFORM_REQUIRES_TYPELESS_RESOURCE_DISCARD_WORKAROUND
+		| ETextureCreateFlags::NoFastClear
+#endif
+		;
 }
 
-FSlatePostProcessor::~FSlatePostProcessor()
+ETextureCreateFlags GetSlateTransientDepthStencilFlags()
 {
-	// Note this is deleted automatically because it implements FDeferredCleanupInterface.
-	for (int32 ui = 0; ui < IntermediateTargetsArray.Num(); ++ui)
-	{
-		IntermediateTargetsArray[ui]->CleanUp();
-	}
- 	IntermediateTargetsArray.Empty();
+	return ETextureCreateFlags::DepthStencilTargetable | ETextureCreateFlags::FastVRAM;
 }
 
+//////////////////////////////////////////////////////////////////////////
 
-// Pixel shader to composite UI over HDR buffer
-class FBlitUIToHDRPSBase : public FGlobalShader
+// Pixel shader to composite UI over HDR buffer prior to doing a blur
+class FCompositeHDRForBlurPS : public FGlobalShader
 {
-	DECLARE_TYPE_LAYOUT(FBlitUIToHDRPSBase, NonVirtual);
 public:
+	DECLARE_GLOBAL_SHADER(FCompositeHDRForBlurPS);
+	SHADER_USE_PARAMETER_STRUCT(FCompositeHDRForBlurPS, FGlobalShader);
+
+	class FUseSRGBEncoding : SHADER_PERMUTATION_BOOL("SCRGB_ENCODING");
+	using FPermutationDomain = TShaderPermutationDomain<FUseSRGBEncoding>;
+
+	BEGIN_SHADER_PARAMETER_STRUCT(FParameters, )
+		SHADER_PARAMETER_RDG_TEXTURE(Texture2D, UITexture)
+		SHADER_PARAMETER_RDG_TEXTURE(Texture2D, UIWriteMaskTexture)
+		SHADER_PARAMETER_RDG_TEXTURE(Texture2D, SceneTexture)
+		SHADER_PARAMETER_SAMPLER(SamplerState, UISampler)
+		SHADER_PARAMETER(float, UILevel)
+		SHADER_PARAMETER(float, UILuminance)
+		SHADER_PARAMETER(FVector2f, UITextureSize)
+		RENDER_TARGET_BINDING_SLOTS()
+	END_SHADER_PARAMETER_STRUCT()
 
 	static bool ShouldCompilePermutation(const FGlobalShaderPermutationParameters& Parameters)
 	{
 		return IsFeatureLevelSupported(Parameters.Platform, ERHIFeatureLevel::SM5) && (RHISupportsGeometryShaders(Parameters.Platform) || RHISupportsVertexShaderLayer(Parameters.Platform));
 	}
 
-	FBlitUIToHDRPSBase(const ShaderMetaType::CompiledShaderInitializerType& Initializer) :
-		FGlobalShader(Initializer)
-	{
-		HDRTexture.Bind(Initializer.ParameterMap, TEXT("SceneTexture"));
-		UITexture.Bind(Initializer.ParameterMap, TEXT("UITexture"));
-		UIWriteMaskTexture.Bind(Initializer.ParameterMap, TEXT("UIWriteMaskTexture"));
-		UISampler.Bind(Initializer.ParameterMap, TEXT("UISampler"));
-		UILevel.Bind(Initializer.ParameterMap, TEXT("UILevel"));
-		UILuminance.Bind(Initializer.ParameterMap, TEXT("UILuminance"));
-		UITextureSize.Bind(Initializer.ParameterMap, TEXT("UITextureSize"));
-	}
-	FBlitUIToHDRPSBase() = default;
-
-	void SetParameters(FRHIBatchedShaderParameters& BatchedParameters, FRHITexture* UITextureRHI, FRHITexture* HDRTextureRHI, FRHITexture* UITextureWriteMaskRHI, const FVector2f& InUITextureSize)
-	{
-		SetTextureParameter(BatchedParameters, HDRTexture, UISampler, TStaticSamplerState<SF_Bilinear>::GetRHI(), HDRTextureRHI);
-		SetTextureParameter(BatchedParameters, UITexture, UISampler, TStaticSamplerState<SF_Bilinear>::GetRHI(), UITextureRHI);
-		static auto CVarHDRUILevel = IConsoleManager::Get().FindConsoleVariable(TEXT("r.HDR.UI.Level"));
-		static auto CVarHDRUILuminance = IConsoleManager::Get().FindConsoleVariable(TEXT("r.HDR.UI.Luminance"));
-		SetShaderValue(BatchedParameters, UILevel, CVarHDRUILevel ? CVarHDRUILevel->GetFloat() : 1.0f);
-		SetShaderValue(BatchedParameters, UILuminance, CVarHDRUILuminance ? CVarHDRUILuminance->GetFloat() : 300.0f);
-		SetShaderValue(BatchedParameters, UITextureSize, InUITextureSize);
-		if (UITextureWriteMaskRHI != nullptr)
-		{
-			SetTextureParameter(BatchedParameters, UIWriteMaskTexture, UITextureWriteMaskRHI);
-		}
-	}
-
 	static void ModifyCompilationEnvironment(const FGlobalShaderPermutationParameters& Parameters, FShaderCompilerEnvironment& OutEnvironment)
 	{
 		FGlobalShader::ModifyCompilationEnvironment(Parameters, OutEnvironment);
-		OutEnvironment.SetDefine(TEXT("BLIT_UI_TO_HDR"), 1);
+		OutEnvironment.SetDefine(TEXT("COMPOSITE_UI_FOR_BLUR_PS"), 1);
 	}
-
-	static const TCHAR* GetSourceFilename()
-	{
-		return TEXT("/Engine/Private/CompositeUIPixelShader.usf");
-	}
-
-	static const TCHAR* GetFunctionName()
-	{
-		return TEXT("BlitUIToHDRPS");
-	}
-
-private:
-	LAYOUT_FIELD(FShaderResourceParameter, HDRTexture);
-	LAYOUT_FIELD(FShaderResourceParameter, UITexture);
-	LAYOUT_FIELD(FShaderResourceParameter, UISampler);
-	LAYOUT_FIELD(FShaderResourceParameter, UIWriteMaskTexture);
-	LAYOUT_FIELD(FShaderParameter, UILevel);
-	LAYOUT_FIELD(FShaderParameter, UILuminance);
-	LAYOUT_FIELD(FShaderParameter, UITextureSize);
 };
 
-IMPLEMENT_TYPE_LAYOUT(FBlitUIToHDRPSBase);
+IMPLEMENT_GLOBAL_SHADER(FCompositeHDRForBlurPS, "/Engine/Private/CompositeUIPixelShader.usf", "CompositeUIForBlur", SF_Pixel);
 
-template<uint32 EncodingType>
-class FBlitUIToHDRPS : public FBlitUIToHDRPSBase
+struct FSlateCompositeHDRForBlurPassInputs
 {
-	DECLARE_SHADER_TYPE(FBlitUIToHDRPS, Global);
-public:
-	FBlitUIToHDRPS() = default;
-	FBlitUIToHDRPS(const ShaderMetaType::CompiledShaderInitializerType& Initializer) : FBlitUIToHDRPSBase(Initializer) {}
-
-	static void ModifyCompilationEnvironment(const FGlobalShaderPermutationParameters& Parameters, FShaderCompilerEnvironment& OutEnvironment)
-	{
-		FBlitUIToHDRPSBase::ModifyCompilationEnvironment(Parameters, OutEnvironment);
-		OutEnvironment.SetDefine(TEXT("SCRGB_ENCODING"), EncodingType);
-	}
+	FIntRect InputRect;
+	FRDGTexture* InputCompositeTexture;
+	FRDGTexture* InputTexture;
+	FIntPoint OutputExtent;
 };
 
-IMPLEMENT_SHADER_TYPE(template<>, FBlitUIToHDRPS<0>, FBlitUIToHDRPS::GetSourceFilename(), FBlitUIToHDRPS::GetFunctionName(), SF_Pixel);
-IMPLEMENT_SHADER_TYPE(template<>, FBlitUIToHDRPS<1>, FBlitUIToHDRPS::GetSourceFilename(), FBlitUIToHDRPS::GetFunctionName(), SF_Pixel);
-
-static void BlitUIToHDRScene(FRHICommandListImmediate& RHICmdList, IRendererModule& RendererModule, const FPostProcessRectParams& RectParams, FTextureRHIRef DestTexture, const FIntPoint& AllocatedSize)
+FScreenPassTexture AddSlateCompositeHDRForBlurPass(FRDGBuilder& GraphBuilder, const FSlateCompositeHDRForBlurPassInputs& Inputs)
 {
-	SCOPED_DRAW_EVENT(RHICmdList, SlatePostProcessBlitUIToHDR);
-
-	FRHITexture* UITexture = RectParams.PostProcessDest == EPostProcessDestination::DestTexture 
-		? RectParams.DestTexture->GetTexture2D() 
-		: RectParams.UITarget->GetRHI();
-
-	RHICmdList.Transition(FRHITransitionInfo(UITexture, ERHIAccess::Unknown, ERHIAccess::SRVMask));
-	TRefCountPtr<IPooledRenderTarget> UITargetRTMask;
-	if (RHISupportsRenderTargetWriteMask(GMaxRHIShaderPlatform) && RectParams.PostProcessDest != EPostProcessDestination::DestTexture)
-	{
-		const auto FeatureLevel = GMaxRHIFeatureLevel;
-		auto ShaderMap = GetGlobalShaderMap(FeatureLevel);
-
-		IPooledRenderTarget* RenderTargets[] = { RectParams.UITarget.GetReference() };
-		FRenderTargetWriteMask::Decode(RHICmdList, ShaderMap, RenderTargets, UITargetRTMask, TexCreate_None, TEXT("UIRTWriteMask"));
-	}
-	FRHITexture* UITargetRTMaskTexture = RHISupportsRenderTargetWriteMask(GMaxRHIShaderPlatform) ? UITargetRTMask->GetRHI() : nullptr;
-
-	// Source is the viewport.  This is the width and height of the viewport backbuffer
-	const int32 SrcTextureWidth = RectParams.SourceTextureSize.X;
-	const int32 SrcTextureHeight = RectParams.SourceTextureSize.Y;
-
-	// Rect of the viewport
-	const FSlateRect& SourceRect = RectParams.SourceRect;
-
-	// Rect of the final destination post process effect (not downsample rect).  This is the area we sample from
-	const FSlateRect& DestRect = RectParams.DestRect;
-
-	FGlobalShaderMap* ShaderMap = GetGlobalShaderMap(GMaxRHIFeatureLevel);
-	TShaderMapRef<FScreenVS> VertexShader(ShaderMap);
-
-	bool bIsSCRGB = (RectParams.SourceTexture->GetDesc().Format == PF_FloatRGBA);
-
-	TShaderRef<FBlitUIToHDRPSBase> PixelShader;  
-	if (bIsSCRGB)
-	{
-		PixelShader = TShaderMapRef<FBlitUIToHDRPS<1> >(ShaderMap);
-	}
-	else
-	{
-		PixelShader = TShaderMapRef<FBlitUIToHDRPS<0> >(ShaderMap);
-	}
-
-	TArray<FRHITransitionInfo> RHITransitionInfos =
-	{
-		FRHITransitionInfo(UITexture, ERHIAccess::Unknown, ERHIAccess::SRVGraphics),
-		FRHITransitionInfo(RectParams.SourceTexture, ERHIAccess::Unknown, ERHIAccess::SRVGraphics),
-		FRHITransitionInfo(DestTexture, ERHIAccess::Unknown, ERHIAccess::RTV)
-	};
-
+	const ERHIFeatureLevel::Type FeatureLevel = GMaxRHIFeatureLevel;
+	FGlobalShaderMap* ShaderMap = GetGlobalShaderMap(FeatureLevel);
+	FRDGTexture* UIWriteMaskTexture = nullptr;
 
-	if (UITargetRTMaskTexture)
+	if (RHISupportsRenderTargetWriteMask(GMaxRHIShaderPlatform))
 	{
-		RHITransitionInfos.Add(FRHITransitionInfo(UITargetRTMaskTexture, ERHIAccess::Unknown, ERHIAccess::SRVGraphics));
+		FRenderTargetWriteMask::Decode(GraphBuilder, ShaderMap, MakeArrayView({ Inputs.InputCompositeTexture }), UIWriteMaskTexture, TexCreate_None, TEXT("UIRTWriteMask"));
 	}
 
-	RHICmdList.Transition(RHITransitionInfos);
-
-	const FVector2f InvSrcTextureSize(1.f / SrcTextureWidth, 1.f / SrcTextureHeight);
-
-	// If using a custom output target, sample the source rect for UV's, else do not as the dest rect may be a subsection
-	const FSlateRect& UVRect = RectParams.PostProcessDest == EPostProcessDestination::DestTexture
-		? SourceRect
-		: DestRect;
-
-	// no guard band is actually needed because GaussianBlurMain already takes ensure that sampling happens inside the rectangle
-	const FVector2f UVStart = FVector2f(UVRect.Left - 0.0f, UVRect.Top - 0.0f) * InvSrcTextureSize;
-	const FVector2f UVEnd = FVector2f(UVRect.Right + 0.0f, UVRect.Bottom + 0.0f) * InvSrcTextureSize;
-	const FVector2f SizeUV = UVEnd - UVStart;
-
-	RHICmdList.SetViewport(0, 0, 0, SrcTextureWidth, SrcTextureHeight, 0.0f);
-	RHICmdList.SetScissorRect(false, 0, 0, 0, 0);
-
-	FRHIRenderPassInfo RPInfo(DestTexture, ERenderTargetActions::Load_Store);
-	RHICmdList.BeginRenderPass(RPInfo, TEXT("BlitUIToHDR"));
-	{
-		FGraphicsPipelineStateInitializer GraphicsPSOInit;
-		RHICmdList.ApplyCachedRenderTargets(GraphicsPSOInit);
-		GraphicsPSOInit.BlendState = TStaticBlendState<CW_RGB>::GetRHI();
-		GraphicsPSOInit.RasterizerState = TStaticRasterizerState<>::GetRHI();
-		GraphicsPSOInit.DepthStencilState = TStaticDepthStencilState<false, CF_Always>::GetRHI();
-
-		GraphicsPSOInit.BoundShaderState.VertexDeclarationRHI = GFilterVertexDeclaration.VertexDeclarationRHI;
-		GraphicsPSOInit.BoundShaderState.VertexShaderRHI = VertexShader.GetVertexShader();
-		GraphicsPSOInit.BoundShaderState.PixelShaderRHI = PixelShader.GetPixelShader();
-		GraphicsPSOInit.PrimitiveType = PT_TriangleList;
-		SetGraphicsPipelineState(RHICmdList, GraphicsPSOInit, 0);
-
-		SetShaderParametersLegacyPS(RHICmdList, PixelShader, UITexture, RectParams.SourceTexture, UITargetRTMaskTexture, FVector2f(SrcTextureWidth, SrcTextureHeight));
-
-		RendererModule.DrawRectangle(
-			RHICmdList,
-			0, 0,
-			(float)AllocatedSize.X, (float)AllocatedSize.Y,
-			UVStart.X, UVStart.Y,
-			SizeUV.X, SizeUV.Y,
-			DestTexture->GetDesc().Extent,
-			FIntPoint(1, 1),
-			VertexShader,
-			EDRF_Default);
-	}
-	RHICmdList.EndRenderPass();
+	FScreenPassRenderTarget Output(
+		GraphBuilder.CreateTexture(
+			FRDGTextureDesc::Create2D(Inputs.OutputExtent, PF_FloatR11G11B10, FClearValueBinding::Black, GetSlateTransientRenderTargetFlags()),
+			TEXT("CompositeHDRUI")),
+		ERenderTargetLoadAction::ENoAction);
+
+	const FScreenPassTextureViewport InputViewport = FScreenPassTextureViewport(Inputs.InputCompositeTexture, Inputs.InputRect);
+	const FScreenPassTextureViewport OutputViewport = FScreenPassTextureViewport(Output);
+
+	FCompositeHDRForBlurPS::FPermutationDomain PermutationVector;
+	PermutationVector.Set<FCompositeHDRForBlurPS::FUseSRGBEncoding>(Inputs.InputTexture->Desc.Format == PF_FloatRGBA);
+
+	FCompositeHDRForBlurPS::FParameters* PassParameters = GraphBuilder.AllocParameters<FCompositeHDRForBlurPS::FParameters>();
+	PassParameters->RenderTargets[0] = Output.GetRenderTargetBinding();
+	PassParameters->SceneTexture = Inputs.InputTexture;
+	PassParameters->UITexture = Inputs.InputCompositeTexture;
+	PassParameters->UIWriteMaskTexture = UIWriteMaskTexture;
+	PassParameters->UISampler = TStaticSamplerState<SF_Bilinear>::GetRHI();
+	PassParameters->UITextureSize = InputViewport.Extent;
+	PassParameters->UILevel = GetSlateHDRUILevel();
+	PassParameters->UILuminance = GetSlateHDRUILuminance();
+
+	TShaderMapRef<FCompositeHDRForBlurPS> PixelShader(ShaderMap, PermutationVector);
+	AddDrawScreenPass(GraphBuilder, RDG_EVENT_NAME("CompositeHDR"), FeatureLevel, OutputViewport, InputViewport, PixelShader, PassParameters);
+	return Output;
 }
 
-void FSlatePostProcessor::BlurRect(FRHICommandListImmediate& RHICmdList, IRendererModule& RendererModule, const FBlurRectParams& Params, const FPostProcessRectParams& RectParams)
+//////////////////////////////////////////////////////////////////////////
+
+class FSlatePostProcessDownsamplePS : public FGlobalShader
 {
-	SCOPE_CYCLE_COUNTER(STAT_SlatePostProcessingRTTime);
-	check(RHICmdList.IsOutsideRenderPass());
+public:
+	DECLARE_GLOBAL_SHADER(FSlatePostProcessDownsamplePS);
+	SHADER_USE_PARAMETER_STRUCT(FSlatePostProcessDownsamplePS, FGlobalShader);
+
+	BEGIN_SHADER_PARAMETER_STRUCT(FParameters, )
+		SHADER_PARAMETER_RDG_TEXTURE(Texture2D, ElementTexture)
+		SHADER_PARAMETER_SAMPLER(SamplerState, ElementTextureSampler)
+		SHADER_PARAMETER(FVector4f, ShaderParams)
+		SHADER_PARAMETER(FVector4f, UVBounds)
+		RENDER_TARGET_BINDING_SLOTS()
+	END_SHADER_PARAMETER_STRUCT()
+};
 
-	TArray<FVector4f> WeightsAndOffsets;
-	const int32 SampleCount = ComputeBlurWeights(Params.KernelSize, Params.Strength, WeightsAndOffsets);
+IMPLEMENT_GLOBAL_SHADER(FSlatePostProcessDownsamplePS, "/Engine/Private/SlatePostProcessPixelShader.usf", "DownsampleMain", SF_Pixel);
 
+struct FSlatePostProcessDownsamplePassInputs
+{
+	FScreenPassTexture InputTexture;
+	FIntPoint OutputExtent;
+};
 
-	const bool bDownsample = Params.DownsampleAmount > 0;
+FScreenPassTexture AddSlatePostProcessDownsamplePass(FRDGBuilder& GraphBuilder, const FSlatePostProcessDownsamplePassInputs& Inputs)
+{
+	const ERHIFeatureLevel::Type FeatureLevel = GMaxRHIFeatureLevel;
+	FGlobalShaderMap* ShaderMap = GetGlobalShaderMap(FeatureLevel);
+	TShaderMapRef<FSlatePostProcessDownsamplePS> PixelShader(ShaderMap);
+
+	const FScreenPassRenderTarget Output(
+		GraphBuilder.CreateTexture(FRDGTextureDesc::Create2D(Inputs.OutputExtent, Inputs.InputTexture.Texture->Desc.Format, FClearValueBinding::None, GetSlateTransientRenderTargetFlags()), TEXT("DownsampleUI")),
+		ERenderTargetLoadAction::ENoAction);
+
+	const FScreenPassTextureViewport InputViewport(Inputs.InputTexture);
+	const FScreenPassTextureViewportParameters InputParameters = GetScreenPassTextureViewportParameters(InputViewport);
+	const FScreenPassTextureViewport OutputViewport(Output);
+
+	FSlatePostProcessDownsamplePS::FParameters* PassParameters = GraphBuilder.AllocParameters<FSlatePostProcessDownsamplePS::FParameters>();
+	PassParameters->RenderTargets[0] = Output.GetRenderTargetBinding();
+	PassParameters->ElementTexture = Inputs.InputTexture.Texture;
+	PassParameters->ElementTextureSampler = TStaticSamplerState<SF_Bilinear>::GetRHI();
+	PassParameters->ShaderParams = FVector4f(InputParameters.ExtentInverse.X, InputParameters.ExtentInverse.Y, 0.0f, 0.0f);
+	PassParameters->UVBounds = FVector4f(InputParameters.UVViewportBilinearMin, InputParameters.UVViewportBilinearMax);
+
+	AddDrawScreenPass(GraphBuilder, RDG_EVENT_NAME("DownsampleUI"), FeatureLevel, OutputViewport, InputViewport, PixelShader, PassParameters);
+	return Output;
+}
 
-	FIntPoint DestRectSize = RectParams.DestRect.GetSize().IntPoint();
-	FIntPoint RequiredSize = bDownsample
-										? FIntPoint(FMath::DivideAndRoundUp(DestRectSize.X, Params.DownsampleAmount), FMath::DivideAndRoundUp(DestRectSize.Y, Params.DownsampleAmount))
-										: DestRectSize;
+//////////////////////////////////////////////////////////////////////////
 
-	// The max size can get ridiculous with large scale values.  Clamp to size of the backbuffer
-	RequiredSize.X = FMath::Min(RequiredSize.X, RectParams.SourceTextureSize.X);
-	RequiredSize.Y = FMath::Min(RequiredSize.Y, RectParams.SourceTextureSize.Y);
-	
-	SCOPED_DRAW_EVENTF(RHICmdList, SlatePostProcess, TEXT("Slate Post Process Blur Background Kernel: %dx%d Size: %dx%d"), SampleCount, SampleCount, RequiredSize.X, RequiredSize.Y);
+enum class ESlatePostProcessUpsampleOutputFormat
+{
+	SDR = 0,
+	HDR_SCRGB,
+	HDR_PQ10,
+	MAX
+};
 
+class FSlatePostProcessUpsamplePS : public FGlobalShader
+{
+public:
+	DECLARE_GLOBAL_SHADER(FSlatePostProcessUpsamplePS);
+	SHADER_USE_PARAMETER_STRUCT(FSlatePostProcessUpsamplePS, FGlobalShader);
+
+	class FUpsampleOutputFormat : SHADER_PERMUTATION_ENUM_CLASS("UPSAMPLE_OUTPUT_FORMAT", ESlatePostProcessUpsampleOutputFormat);
+	using FPermutationDomain = TShaderPermutationDomain<FUpsampleOutputFormat>;
+
+	BEGIN_SHADER_PARAMETER_STRUCT(FParameters, )
+		SHADER_PARAMETER_RDG_TEXTURE(Texture2D, ElementTexture)
+		SHADER_PARAMETER_SAMPLER(SamplerState, ElementTextureSampler)
+		SHADER_PARAMETER(FVector4f, ShaderParams)
+		SHADER_PARAMETER(FVector4f, ShaderParams2)
+		RENDER_TARGET_BINDING_SLOTS()
+	END_SHADER_PARAMETER_STRUCT()
+};
 
-	const FIntPoint DownsampleSize = RequiredSize;
+IMPLEMENT_GLOBAL_SHADER(FSlatePostProcessUpsamplePS, "/Engine/Private/SlatePostProcessPixelShader.usf", "UpsampleMain", SF_Pixel);
 
-	EPixelFormat IntermediatePixelFormat = RectParams.SourceTexture->GetDesc().Format;
+struct FSlatePostProcessUpsampleInputs
+{
+	FScreenPassTexture InputTexture;
+	FRDGTexture* OutputTextureToClear = nullptr;
+	FRDGTexture* OutputTexture = nullptr;
+	ERenderTargetLoadAction OutputLoadAction = ERenderTargetLoadAction::ELoad;
 
-	bool bIsHDRSource = RectParams.UITarget.IsValid() && RectParams.UITarget->GetRHI() != RectParams.SourceTexture.GetReference();
-	if (bIsHDRSource)
-	{
-		IntermediatePixelFormat = PF_FloatR11G11B10;
-	}
+	const FSlateClippingOp* ClippingOp = nullptr;
+	const FDepthStencilBinding* ClippingStencilBinding = nullptr;
 
-	FSlatePostProcessResource* IntermediateTargets = FindSlatePostProcessResource(IntermediateTargetsArray, IntermediatePixelFormat);
-	IntermediateTargets->Update(RequiredSize, IntermediatePixelFormat);
+	FIntRect OutputRect;
+	FVector4f CornerRadius = FVector4f::Zero();
+};
 
-	if (bIsHDRSource)
-	{
-		// in HDR mode, we are going to blur SourceTexture, but still need to take into account the UI already rendered. Blit UI into HDR target
-		BlitUIToHDRScene(RHICmdList, RendererModule, RectParams, IntermediateTargets->GetRenderTarget(0), RequiredSize);
-	}
+void AddSlatePostProcessUpsamplePass(FRDGBuilder& GraphBuilder, const FSlatePostProcessUpsampleInputs& Inputs)
+{
+	FSlatePostProcessUpsamplePS::FParameters* PassParameters = GraphBuilder.AllocParameters<FSlatePostProcessUpsamplePS::FParameters>();
+	PassParameters->RenderTargets[0] = FRenderTargetBinding(Inputs.OutputTexture, Inputs.OutputLoadAction);
 
-	else if(bDownsample)
+	if (Inputs.ClippingStencilBinding)
 	{
-		DownsampleRect(RHICmdList, RendererModule, RectParams, DownsampleSize, IntermediateTargets);
+		PassParameters->RenderTargets.DepthStencil = *Inputs.ClippingStencilBinding;
 	}
-
-	FSamplerStateRHIRef BilinearClamp = TStaticSamplerState<SF_Bilinear, AM_Clamp, AM_Clamp, AM_Clamp>::GetRHI();
-
-#if 1
-	FGlobalShaderMap* ShaderMap = GetGlobalShaderMap(GMaxRHIFeatureLevel);
-	check(ShaderMap);
-
-	TShaderMapRef<FScreenVS> VertexShader(ShaderMap);
-	TShaderMapRef<FSlatePostProcessBlurPS> PixelShader(ShaderMap);
-
-	const int32 SrcTextureWidth = RectParams.SourceTextureSize.X;
-	const int32 SrcTextureHeight = RectParams.SourceTextureSize.Y;
-
-	const int32 DestTextureWidth = IntermediateTargets->GetWidth();
-	const int32 DestTextureHeight = IntermediateTargets->GetHeight();
-
-	const FSlateRect& SourceRect = RectParams.SourceRect;
-	const FSlateRect& DestRect = RectParams.DestRect;
-
-	FVertexDeclarationRHIRef VertexDecl = GFilterVertexDeclaration.VertexDeclarationRHI;
-	check(IsValidRef(VertexDecl));
-	
-	FGraphicsPipelineStateInitializer GraphicsPSOInit;
-	GraphicsPSOInit.BlendState = TStaticBlendState<>::GetRHI();
-	GraphicsPSOInit.RasterizerState = TStaticRasterizerState<>::GetRHI();
-	GraphicsPSOInit.DepthStencilState = TStaticDepthStencilState<false, CF_Always>::GetRHI();
-
-	RHICmdList.SetScissorRect(false, 0, 0, 0, 0);
-	RHICmdList.SetViewport(0.f, 0.f, 0.f, (float)DestTextureWidth, (float)DestTextureHeight, 0.0f);
 	
-	const FVector2f InvBufferSize = FVector2f(1.0f / (float)DestTextureWidth, 1.0f / (float)DestTextureHeight);
-	const FVector2f HalfTexelOffset = FVector2f(0.5f/ (float)DestTextureWidth, 0.5f/ (float)DestTextureHeight);
+	ESlatePostProcessUpsampleOutputFormat OutputFormat = ESlatePostProcessUpsampleOutputFormat::SDR;
 
-	for (int32 PassIndex = 0; PassIndex < 2; ++PassIndex)
+	if (Inputs.OutputTextureToClear)
 	{
-		// First pass render to the render target with the post process fx
-		if (PassIndex == 0)
-		{
-			FTextureRHIRef SourceTexture = (bDownsample || bIsHDRSource) ? IntermediateTargets->GetRenderTarget(0) : RectParams.SourceTexture;
-			FTextureRHIRef DestTexture = IntermediateTargets->GetRenderTarget(1);
-
-			RHICmdList.Transition(FRHITransitionInfo(SourceTexture, ERHIAccess::Unknown, ERHIAccess::SRVGraphics));
-			RHICmdList.Transition(FRHITransitionInfo(DestTexture, ERHIAccess::Unknown, ERHIAccess::RTV));
-
-			FRHIRenderPassInfo RPInfo(DestTexture, ERenderTargetActions::Load_Store);
-			RHICmdList.BeginRenderPass(RPInfo, TEXT("SlateBlurRectPass0"));
-			{
-				RHICmdList.ApplyCachedRenderTargets(GraphicsPSOInit);
-
-				GraphicsPSOInit.BoundShaderState.VertexDeclarationRHI = VertexDecl;
-				GraphicsPSOInit.BoundShaderState.VertexShaderRHI = VertexShader.GetVertexShader();
-				GraphicsPSOInit.BoundShaderState.PixelShaderRHI = PixelShader.GetPixelShader();
-				GraphicsPSOInit.PrimitiveType = PT_TriangleList;
-				SetGraphicsPipelineState(RHICmdList, GraphicsPSOInit, 0);
-
-				FRHIBatchedShaderParameters& BatchedParameters = RHICmdList.GetScratchShaderParameters();
-
-				PixelShader->SetWeightsAndOffsets(BatchedParameters, WeightsAndOffsets, SampleCount);
-				PixelShader->SetTexture(BatchedParameters, SourceTexture, BilinearClamp);
-
-				if (bDownsample || bIsHDRSource)
-				{
-					PixelShader->SetUVBounds(BatchedParameters, FVector4f(FVector2f::ZeroVector, FVector2f((float)DownsampleSize.X / DestTextureWidth, (float)DownsampleSize.Y / DestTextureHeight) - HalfTexelOffset));
-					PixelShader->SetBufferSizeAndDirection(BatchedParameters, InvBufferSize, FVector2f(1.f, 0.f));
-
-					RHICmdList.SetBatchedShaderParameters(PixelShader.GetPixelShader(), BatchedParameters);
-
-					RendererModule.DrawRectangle(
-						RHICmdList,
-						0.f, 0.f,
-						(float)DownsampleSize.X, (float)DownsampleSize.Y,
-						0, 0,
-						(float)DownsampleSize.X, (float)DownsampleSize.Y,
-						FIntPoint(DestTextureWidth, DestTextureHeight),
-						FIntPoint(DestTextureWidth, DestTextureHeight),
-						VertexShader,
-						EDRF_Default);
-				}
-				else
-				{
-					const FVector2f InvSrcTextureSize(1.f / SrcTextureWidth, 1.f / SrcTextureHeight);
-
-					// If using a custom output target, sample the source rect for UV's, else do not as the dest rect may be a subsection
-					const FSlateRect& UVRect = RectParams.PostProcessDest == EPostProcessDestination::DestTexture
-						? SourceRect
-						: DestRect;
-
-					const FVector2f UVStart = FVector2f(UVRect.Left, UVRect.Top) * InvSrcTextureSize;
-					const FVector2f UVEnd = FVector2f(UVRect.Right, UVRect.Bottom) * InvSrcTextureSize;
-					const FVector2f SizeUV = UVEnd - UVStart;
-
-					PixelShader->SetUVBounds(BatchedParameters, FVector4f(UVStart, UVEnd));
-					PixelShader->SetBufferSizeAndDirection(BatchedParameters, InvSrcTextureSize, FVector2f(1.f, 0.f));
-
-					RHICmdList.SetBatchedShaderParameters(PixelShader.GetPixelShader(), BatchedParameters);
-
-					RendererModule.DrawRectangle(
-						RHICmdList,
-						0.f, 0.f,
-						(float)RequiredSize.X, (float)RequiredSize.Y,
-						UVStart.X, UVStart.Y,
-						SizeUV.X, SizeUV.Y,
-						FIntPoint(DestTextureWidth, DestTextureHeight),
-						FIntPoint(1, 1),
-						VertexShader,
-						EDRF_Default);
-				}
-			}
-			RHICmdList.EndRenderPass();
-		}
-		else
-		{
-			FTextureRHIRef SourceTexture = IntermediateTargets->GetRenderTarget(1);
-			FTextureRHIRef DestTexture = IntermediateTargets->GetRenderTarget(0);
-
-			RHICmdList.Transition(FRHITransitionInfo(SourceTexture, ERHIAccess::Unknown, ERHIAccess::SRVGraphics));
-			RHICmdList.Transition(FRHITransitionInfo(DestTexture, ERHIAccess::Unknown, ERHIAccess::RTV));
-
-			FRHIRenderPassInfo RPInfo(DestTexture, ERenderTargetActions::Load_Store);
-			RHICmdList.BeginRenderPass(RPInfo, TEXT("SlateBlurRect"));
-			{
-				RHICmdList.ApplyCachedRenderTargets(GraphicsPSOInit);
-
-				GraphicsPSOInit.BoundShaderState.VertexDeclarationRHI = VertexDecl;
-				GraphicsPSOInit.BoundShaderState.VertexShaderRHI = VertexShader.GetVertexShader();
-				GraphicsPSOInit.BoundShaderState.PixelShaderRHI = PixelShader.GetPixelShader();
-				GraphicsPSOInit.PrimitiveType = PT_TriangleList;
-				SetGraphicsPipelineState(RHICmdList, GraphicsPSOInit, 0);
-
-				FRHIBatchedShaderParameters& BatchedParameters = RHICmdList.GetScratchShaderParameters();
-
-				PixelShader->SetWeightsAndOffsets(BatchedParameters, WeightsAndOffsets, SampleCount);
-				PixelShader->SetUVBounds(BatchedParameters, FVector4f(FVector2f::ZeroVector, FVector2f((float)DownsampleSize.X / DestTextureWidth, (float)DownsampleSize.Y / DestTextureHeight) - HalfTexelOffset));
-				PixelShader->SetTexture(BatchedParameters, SourceTexture, BilinearClamp);
-				PixelShader->SetBufferSizeAndDirection(BatchedParameters, InvBufferSize, FVector2f(0.f, 1.f));
-
-				RHICmdList.SetBatchedShaderParameters(PixelShader.GetPixelShader(), BatchedParameters);
-
-				RendererModule.DrawRectangle(
-					RHICmdList,
-					0.f, 0.f,
-					(float)DownsampleSize.X, (float)DownsampleSize.Y,
-					0.f, 0.f,
-					(float)DownsampleSize.X, (float)DownsampleSize.Y,
-					FIntPoint(DestTextureWidth, DestTextureHeight),
-					FIntPoint(DestTextureWidth, DestTextureHeight),
-					VertexShader,
-					EDRF_Default);
-			}
-			RHICmdList.EndRenderPass();
-		}	
-	}
-
-#endif
-
-	UpsampleRect(RHICmdList, RendererModule, RectParams, DownsampleSize, BilinearClamp, IntermediateTargets);
-}
+		OutputFormat = Inputs.OutputTexture->Desc.Format == PF_FloatRGBA
+			? ESlatePostProcessUpsampleOutputFormat::HDR_SCRGB
+			: ESlatePostProcessUpsampleOutputFormat::HDR_PQ10;
 
-void FSlatePostProcessor::ColorDeficiency(FRHICommandListImmediate& RHICmdList, IRendererModule& RendererModule, const FPostProcessRectParams& RectParams)
-{
-	SCOPE_CYCLE_COUNTER(STAT_SlateColorDeficiencyRTTime);
-
-	FIntPoint DestRectSize = RectParams.DestRect.GetSize().IntPoint();
-	FIntPoint RequiredSize = DestRectSize;
-
-	FSlatePostProcessResource* IntermediateTargets = FindSlatePostProcessResource(IntermediateTargetsArray, RectParams.SourceTexture->GetDesc().Format);
-
-	IntermediateTargets->Update(RequiredSize, RectParams.SourceTexture->GetDesc().Format);
+		PassParameters->RenderTargets[1] = FRenderTargetBinding(Inputs.OutputTextureToClear, Inputs.OutputLoadAction);
+	}
 
-#if 1
-	FSamplerStateRHIRef PointClamp = TStaticSamplerState<SF_Point, AM_Clamp, AM_Clamp, AM_Clamp>::GetRHI();
+	FSlatePostProcessUpsamplePS::FPermutationDomain PermutationVector;
+	PermutationVector.Set<FSlatePostProcessUpsamplePS::FUpsampleOutputFormat>(OutputFormat);
 
-	FGlobalShaderMap* ShaderMap = GetGlobalShaderMap(GMaxRHIFeatureLevel);
-	check(ShaderMap);
+	const ERHIFeatureLevel::Type FeatureLevel = GMaxRHIFeatureLevel;
+	FGlobalShaderMap* ShaderMap = GetGlobalShaderMap(FeatureLevel);
 
-	TShaderMapRef<FScreenVS> VertexShader(ShaderMap);
-	TShaderMapRef<FSlatePostProcessColorDeficiencyPS> PixelShader(ShaderMap);
+	TShaderMapRef<FScreenPassVS> VertexShader(ShaderMap);
+	TShaderMapRef<FSlatePostProcessUpsamplePS> PixelShader(ShaderMap, PermutationVector);
 
-	const int32 SrcTextureWidth = RectParams.SourceTextureSize.X;
-	const int32 SrcTextureHeight = RectParams.SourceTextureSize.Y;
+	const FScreenPassTextureViewport InputViewport(Inputs.InputTexture);
+	const FScreenPassTextureViewport OutputViewport(Inputs.OutputTexture, Inputs.OutputRect);
+	const FScreenPassTextureViewportParameters InputParameters = GetScreenPassTextureViewportParameters(InputViewport);
 
-	const int32 DestTextureWidth = IntermediateTargets->GetWidth();
-	const int32 DestTextureHeight = IntermediateTargets->GetHeight();
+	PassParameters->ElementTexture = Inputs.InputTexture.Texture;
+	PassParameters->ElementTextureSampler = Inputs.InputTexture.ViewRect == Inputs.OutputRect
+		? TStaticSamplerState<SF_Point, AM_Clamp, AM_Clamp, AM_Clamp>::GetRHI()
+		: TStaticSamplerState<SF_Bilinear, AM_Clamp, AM_Clamp, AM_Clamp>::GetRHI();
 
-	const FSlateRect& SourceRect = RectParams.SourceRect;
-	const FSlateRect& DestRect = RectParams.DestRect;
+	PassParameters->ShaderParams = FVector4f(InputParameters.ViewportSize, InputParameters.UVViewportSize);
+	PassParameters->ShaderParams2 = Inputs.CornerRadius;
 
-	FVertexDeclarationRHIRef VertexDecl = GFilterVertexDeclaration.VertexDeclarationRHI;
-	check(IsValidRef(VertexDecl));
+	FRHIBlendState* BlendState = Inputs.CornerRadius == FVector4f::Zero() ? TStaticBlendState<>::GetRHI() : TStaticBlendState<CW_RGBA, BO_Add, BF_SourceAlpha, BF_InverseSourceAlpha, BO_Add, BF_One, BF_InverseSourceAlpha>::GetRHI();
 
-	FGraphicsPipelineStateInitializer GraphicsPSOInit;
-	GraphicsPSOInit.BlendState = TStaticBlendState<>::GetRHI();
-	GraphicsPSOInit.RasterizerState = TStaticRasterizerState<>::GetRHI();
-	GraphicsPSOInit.DepthStencilState = TStaticDepthStencilState<false, CF_Always>::GetRHI();
+	FScreenPassPipelineState PipelineState(VertexShader, PixelShader, BlendState);
+	GetSlateClippingPipelineState(Inputs.ClippingOp, PipelineState.DepthStencilState, PipelineState.StencilRef);
 
-	RHICmdList.SetScissorRect(false, 0, 0, 0, 0);
-	RHICmdList.SetViewport(0.f, 0.f, 0.f, (float)DestTextureWidth, (float)DestTextureHeight, 0.0f);
-
-	// 
+	GraphBuilder.AddPass(
+		RDG_EVENT_NAME("Upsample"),
+		PassParameters,
+		ERDGPassFlags::Raster,
+		[OutputViewport, InputViewport, PipelineState, PixelShader, ClippingOp = Inputs.ClippingOp, PassParameters](FRHICommandList& RHICmdList)
 	{
-		FTextureRHIRef SourceTexture = RectParams.SourceTexture;
-		FTextureRHIRef DestTexture = IntermediateTargets->GetRenderTarget(0);
-
-		RHICmdList.Transition(FRHITransitionInfo(SourceTexture, ERHIAccess::Unknown, ERHIAccess::SRVGraphics));
-		RHICmdList.Transition(FRHITransitionInfo(DestTexture, ERHIAccess::Unknown, ERHIAccess::RTV));
+		// Stencil clipping will issue its own draw calls.
+		SetSlateClipping(RHICmdList, ClippingOp, OutputViewport.Rect);
 
-		FRHIRenderPassInfo RPInfo(DestTexture, ERenderTargetActions::Load_Store);
-		RHICmdList.BeginRenderPass(RPInfo, TEXT("ColorDeficiency"));
+		DrawScreenPass(RHICmdList, FScreenPassViewInfo(), OutputViewport, InputViewport, PipelineState, EScreenPassDrawFlags::None, [&](FRHICommandList&)
 		{
-			RHICmdList.ApplyCachedRenderTargets(GraphicsPSOInit);
-
-			GraphicsPSOInit.BoundShaderState.VertexDeclarationRHI = VertexDecl;
-			GraphicsPSOInit.BoundShaderState.VertexShaderRHI = VertexShader.GetVertexShader();
-			GraphicsPSOInit.BoundShaderState.PixelShaderRHI = PixelShader.GetPixelShader();
-			GraphicsPSOInit.PrimitiveType = PT_TriangleList;
-			SetGraphicsPipelineState(RHICmdList, GraphicsPSOInit, 0);
-
-			FRHIBatchedShaderParameters& BatchedParameters = RHICmdList.GetScratchShaderParameters();
-
-			PixelShader->SetColorRules(BatchedParameters, GSlateColorDeficiencyCorrection, GSlateColorDeficiencyType, GSlateColorDeficiencySeverity);
-			PixelShader->SetShowCorrectionWithDeficiency(BatchedParameters, GSlateShowColorDeficiencyCorrectionWithDeficiency);
-			PixelShader->SetTexture(BatchedParameters, SourceTexture, PointClamp);
-
-			RHICmdList.SetBatchedShaderParameters(PixelShader.GetPixelShader(), BatchedParameters);
-
-			RendererModule.DrawRectangle(
-				RHICmdList,
-				0.f, 0,
-				(float)RequiredSize.X, (float)RequiredSize.Y,
-				0.f, 0.f,
-				1.f, 1.f,
-				FIntPoint(DestTextureWidth, DestTextureHeight),
-				FIntPoint(1, 1),
-				VertexShader,
-				EDRF_Default);
-		}
-		RHICmdList.EndRenderPass();
-	}
-
-	const FIntPoint DownsampleSize = RequiredSize;
-	UpsampleRect(RHICmdList, RendererModule, RectParams, DownsampleSize, PointClamp, IntermediateTargets);
-
-#endif
+			SetShaderParameters(RHICmdList, PixelShader, PixelShader.GetPixelShader(), *PassParameters);
+		});
+	});
 }
 
-void FSlatePostProcessor::ReleaseRenderTargets()
-{
-	check(IsInGameThread());
-	// Only release the resource not delete it.  Deleting it could cause issues on any RHI thread
-	for (int32 ui = 0; ui < IntermediateTargetsArray.Num(); ++ui)
-	{
-		BeginReleaseResource(IntermediateTargetsArray[ui]);
-	}
-}
+//////////////////////////////////////////////////////////////////////////
 
-
-void FSlatePostProcessor::DownsampleRect(FRHICommandListImmediate& RHICmdList, IRendererModule& RendererModule, const FPostProcessRectParams& Params, const FIntPoint& DownsampleSize, FSlatePostProcessResource* IntermediateTargets)
+class FSlatePostProcessBlurPS : public FGlobalShader
 {
-	SCOPED_DRAW_EVENT(RHICmdList, SlatePostProcessDownsample);
-
-	// Source is the viewport.  This is the width and height of the viewport backbuffer
-	const int32 SrcTextureWidth = Params.SourceTextureSize.X;
-	const int32 SrcTextureHeight = Params.SourceTextureSize.Y;
-
-	// Dest is the destination quad for the downsample
-	const int32 DestTextureWidth = IntermediateTargets->GetWidth();
-	const int32 DestTextureHeight = IntermediateTargets->GetHeight();
-
-	// Rect of the viewport
-	const FSlateRect& SourceRect = Params.SourceRect;
-
-	// Rect of the final destination post process effect (not downsample rect).  This is the area we sample from
-	const FSlateRect& DestRect = Params.DestRect;
-
-	FGlobalShaderMap* ShaderMap = GetGlobalShaderMap(GMaxRHIFeatureLevel);
-	TShaderMapRef<FScreenVS> VertexShader(ShaderMap);
+public:
+	static const int32 MAX_BLUR_SAMPLES = 127 / 2;
+
+	DECLARE_GLOBAL_SHADER(FSlatePostProcessBlurPS);
+	SHADER_USE_PARAMETER_STRUCT(FSlatePostProcessBlurPS, FGlobalShader);
+
+	BEGIN_SHADER_PARAMETER_STRUCT(FParameters, )
+		SHADER_PARAMETER_RDG_TEXTURE(Texture2D, ElementTexture)
+		SHADER_PARAMETER_SAMPLER(SamplerState, ElementTextureSampler)
+		SHADER_PARAMETER_ARRAY(FVector4f, WeightAndOffsets, [MAX_BLUR_SAMPLES])
+		SHADER_PARAMETER(uint32, SampleCount)
+		SHADER_PARAMETER(FVector4f, BufferSizeAndDirection)
+		SHADER_PARAMETER(FVector4f, UVBounds)
+		RENDER_TARGET_BINDING_SLOTS()
+	END_SHADER_PARAMETER_STRUCT()
+};
 
-	FSamplerStateRHIRef BilinearClamp = TStaticSamplerState<SF_Bilinear, AM_Clamp, AM_Clamp, AM_Clamp>::GetRHI();
+IMPLEMENT_GLOBAL_SHADER(FSlatePostProcessBlurPS, "/Engine/Private/SlatePostProcessPixelShader.usf", "GaussianBlurMain", SF_Pixel);
 
-	FTextureRHIRef DestTexture = IntermediateTargets->GetRenderTarget(0);
+void AddSlatePostProcessBlurPass(FRDGBuilder& GraphBuilder, const FSlatePostProcessBlurPassInputs& Inputs)
+{
+	RDG_EVENT_SCOPE(GraphBuilder, "GaussianBlur");
 
-	// Downsample and store in intermediate texture
+	const auto GetWeight = [](float Dist, float Strength)
 	{
-		TShaderMapRef<FSlatePostProcessDownsamplePS> PixelShader(ShaderMap);
-
-		RHICmdList.Transition(FRHITransitionInfo(Params.SourceTexture, ERHIAccess::Unknown, ERHIAccess::SRVGraphics));
-		RHICmdList.Transition(FRHITransitionInfo(DestTexture, ERHIAccess::Unknown, ERHIAccess::RTV));
+		float Strength2 = Strength*Strength;
+		return (1.0f / FMath::Sqrt(2 * PI*Strength2))*FMath::Exp(-(Dist*Dist) / (2 * Strength2));
+	};
 
-		const FVector2f InvSrcTextureSize(1.f/SrcTextureWidth, 1.f/SrcTextureHeight);
+	const auto GetWeightsAndOffset = [GetWeight](float Dist, float Sigma)
+	{
+		float Offset1 = Dist;
+		float Weight1 = GetWeight(Offset1, Sigma);
 
-		// If using a custom output target, sample the source rect for UV's, else do not as the dest rect may be a subsection
-		const FSlateRect& UVRect = Params.PostProcessDest == EPostProcessDestination::DestTexture
-			? SourceRect
-			: DestRect;
+		float Offset2 = Dist + 1;
+		float Weight2 = GetWeight(Offset2, Sigma);
 
-		const FVector2f UVStart = FVector2f(UVRect.Left, UVRect.Top) * InvSrcTextureSize;
-		const FVector2f UVEnd = FVector2f(UVRect.Right, UVRect.Bottom) * InvSrcTextureSize;
-		const FVector2f SizeUV = UVEnd - UVStart;
-		
-		RHICmdList.SetViewport(0.f, 0.f, 0.f, (float)DestTextureWidth, (float)DestTextureHeight, 0.0f);
-		RHICmdList.SetScissorRect(false, 0, 0, 0, 0);
+		float TotalWeight = Weight1 + Weight2;
+		float Offset = 0;
 
-		FRHIRenderPassInfo RPInfo(DestTexture, ERenderTargetActions::Load_Store);
-		RHICmdList.BeginRenderPass(RPInfo, TEXT("DownsampleRect"));
+		if (TotalWeight > 0)
 		{
-			FGraphicsPipelineStateInitializer GraphicsPSOInit;
-			RHICmdList.ApplyCachedRenderTargets(GraphicsPSOInit);
-			GraphicsPSOInit.BlendState = TStaticBlendState<>::GetRHI();
-			GraphicsPSOInit.RasterizerState = TStaticRasterizerState<>::GetRHI();
-			GraphicsPSOInit.DepthStencilState = TStaticDepthStencilState<false, CF_Always>::GetRHI();
-
-			GraphicsPSOInit.BoundShaderState.VertexDeclarationRHI = GFilterVertexDeclaration.VertexDeclarationRHI;
-			GraphicsPSOInit.BoundShaderState.VertexShaderRHI = VertexShader.GetVertexShader();
-			GraphicsPSOInit.BoundShaderState.PixelShaderRHI = PixelShader.GetPixelShader();
-			GraphicsPSOInit.PrimitiveType = PT_TriangleList;
-			SetGraphicsPipelineState(RHICmdList, GraphicsPSOInit, 0);
-
-			FRHIBatchedShaderParameters& BatchedParameters = RHICmdList.GetScratchShaderParameters();
-
-			PixelShader->SetShaderParams(BatchedParameters, FShaderParams::MakePixelShaderParams(FVector4f(InvSrcTextureSize.X, InvSrcTextureSize.Y, 0, 0)));
-			PixelShader->SetUVBounds(BatchedParameters, FVector4f(UVStart, UVEnd));
-			PixelShader->SetTexture(BatchedParameters, Params.SourceTexture, BilinearClamp);
-
-			RHICmdList.SetBatchedShaderParameters(PixelShader.GetPixelShader(), BatchedParameters);
-
-			RendererModule.DrawRectangle(
-				RHICmdList,
-				0.f, 0.f,
-				(float)DownsampleSize.X, (float)DownsampleSize.Y,
-				UVStart.X, UVStart.Y,
-				SizeUV.X, SizeUV.Y,
-				FIntPoint(DestTextureWidth, DestTextureHeight),
-				FIntPoint(1, 1),
-				VertexShader,
-				EDRF_Default);
+			Offset = (Weight1 * Offset1 + Weight2 * Offset2) / TotalWeight;
 		}
-		RHICmdList.EndRenderPass();
-	}
-	
-	// Testing only
-#if 0
-	UpsampleRect(RHICmdList, RendererModule, Params, DownsampleSize, IntermediateTargets);
-#endif
-}
-
-void FSlatePostProcessor::UpsampleRect(FRHICommandListImmediate& RHICmdList, IRendererModule& RendererModule, const FPostProcessRectParams& Params, const FIntPoint& DownsampleSize, FSamplerStateRHIRef& Sampler, FSlatePostProcessResource* IntermediateTargets)
-{
-	SCOPED_DRAW_EVENT(RHICmdList, SlatePostProcessUpsample);
-
-	const FVector4f Zero(0.f, 0.f, 0.f, 0.f);
-
-	FGraphicsPipelineStateInitializer GraphicsPSOInit;
-	GraphicsPSOInit.BlendState = Params.CornerRadius == Zero ? TStaticBlendState<>::GetRHI() : TStaticBlendState<CW_RGBA, BO_Add, BF_SourceAlpha, BF_InverseSourceAlpha, BO_Add, BF_One, BF_InverseSourceAlpha>::GetRHI();
-	GraphicsPSOInit.RasterizerState = TStaticRasterizerState<>::GetRHI();
-	GraphicsPSOInit.DepthStencilState = TStaticDepthStencilState<false, CF_Always>::GetRHI();
-
-	// Original source texture is now the destination texture
-	FTextureRHIRef DestTexture = Params.PostProcessDest == EPostProcessDestination::DestTexture && Params.DestTexture 
-		? Params.DestTexture 
-		: Params.SourceTexture;
-	const int32 DestTextureWidth = DestTexture->GetSizeX();
-	const int32 DestTextureHeight = DestTexture->GetSizeY();
-
-	const int32 DownsampledWidth = DownsampleSize.X;
-	const int32 DownsampledHeight = DownsampleSize.Y;
 
-	// Source texture is the texture that was originally downsampled
-	FTextureRHIRef SrcTexture = IntermediateTargets->GetRenderTarget(0);
-	const int32 SrcTextureWidth = IntermediateTargets->GetWidth();
-	const int32 SrcTextureHeight = IntermediateTargets->GetHeight();
-
-	const FSlateRect& SourceRect = Params.SourceRect;
-	const FSlateRect& DestRect = Params.DestRect;
-
-	FGlobalShaderMap* ShaderMap = GetGlobalShaderMap(GMaxRHIFeatureLevel);
-	TShaderMapRef<FScreenVS> VertexShader(ShaderMap);
+		return FVector2f(TotalWeight, Offset);
+	};
 
-	RHICmdList.SetViewport(0.f, 0.f, 0.f, (float)DestTextureWidth, (float)DestTextureHeight, 0.0f);
+	const int32 SampleCount = FMath::DivideAndRoundUp(Inputs.KernelSize, 2u);
 
-	// Perform Writable transitions first
+	// We need half of the sample count array because we're packing two samples into one float;
+	TArray<FVector4f, FRDGArrayAllocator> WeightsAndOffsets;
+	WeightsAndOffsets.AddUninitialized(SampleCount % 2 == 0 ? SampleCount / 2 : SampleCount / 2 + 1);
+	WeightsAndOffsets[0] = FVector4f(FVector2f(GetWeight(0, Inputs.Strength), 0), GetWeightsAndOffset(1, Inputs.Strength) );
 
-	TArray<FRHITransitionInfo> RHITransitionInfos =
+	for (uint32 X = 3, SampleIndex = 1; X < Inputs.KernelSize; X += 4, SampleIndex++)
 	{
-		FRHITransitionInfo(SrcTexture, ERHIAccess::Unknown, ERHIAccess::SRVGraphics),
-		FRHITransitionInfo(DestTexture, ERHIAccess::Unknown, ERHIAccess::RTV)
-	};
+		WeightsAndOffsets[SampleIndex] = FVector4f(GetWeightsAndOffset((float)X, Inputs.Strength), GetWeightsAndOffset((float)(X + 2), Inputs.Strength));
+	}
 
-	FRHIRenderPassInfo RPInfo(DestTexture, ERenderTargetActions::Load_Store);
+	FScreenPassTextureViewport OutputTextureViewport(Inputs.InputRect.Size());
 
-	bool bHasMRT = false;
-	bool bIsSCRGB = false;
+	const EPixelFormat InputPixelFormat = Inputs.InputTexture->Desc.Format;
 
-	FRHITexture* UITargetTexture = Params.UITarget.IsValid() ? Params.UITarget->GetRHI() : nullptr;
-	UITargetTexture = Params.PostProcessDest == EPostProcessDestination::DestTexture && Params.DestTexture 
-		? Params.DestTexture->GetTexture2D() 
-		: UITargetTexture;
+	// Defaults to the input UI texture unless a downsample / composite pass is needed.
+	FScreenPassTexture BlurInputTexture(Inputs.InputTexture, Inputs.InputRect);
 
-	if (UITargetTexture != nullptr && DestTexture != UITargetTexture)
+	// Need to composite the HDR scene texture with a separate SDR UI texture (which also does a downsample).
+	if (Inputs.SDRCompositeUITexture)
 	{
-		RPInfo.ColorRenderTargets[1].RenderTarget = UITargetTexture;
-		RPInfo.ColorRenderTargets[1].ArraySlice = -1;
-		RPInfo.ColorRenderTargets[1].Action = ERenderTargetActions::Load_Store;
-		bHasMRT = true;
-		bIsSCRGB = (DestTexture->GetDesc().Format == PF_FloatRGBA);
-		RHITransitionInfos.Add(FRHITransitionInfo(UITargetTexture, ERHIAccess::Unknown, ERHIAccess::RTV));
-	}
+		FSlateCompositeHDRForBlurPassInputs CompositeInputs;
+		CompositeInputs.InputRect = Inputs.InputRect;
+		CompositeInputs.InputTexture = Inputs.InputTexture;
+		CompositeInputs.InputCompositeTexture = Inputs.SDRCompositeUITexture;
+		CompositeInputs.OutputExtent = OutputTextureViewport.Extent;
 
-	RHICmdList.Transition(RHITransitionInfos);
-
-	RHICmdList.BeginRenderPass(RPInfo, TEXT("UpsampleRect"));
+		BlurInputTexture = AddSlateCompositeHDRForBlurPass(GraphBuilder, CompositeInputs);
+	}
+	// Need to do an explicit downsample pass.
+	else if (Inputs.DownsampleAmount > 0)
 	{
-		RHICmdList.ApplyCachedRenderTargets(GraphicsPSOInit);
-		if (Params.RestoreStateFunc)
-		{
-			// This can potentially end and restart a renderpass.
-			// #todo refactor so that we only start one renderpass here. Right now RestoreStateFunc may call UpdateScissorRect which requires an open renderpass.
-			Params.RestoreStateFunc(RHICmdList, GraphicsPSOInit, RPInfo);
-		}
+		OutputTextureViewport = FScreenPassTextureViewport(GetDownscaledExtent(Inputs.InputRect.Size(), Inputs.DownsampleAmount));
 
-		TShaderRef<FSlateElementPS> PixelShader;
-		if (bHasMRT)
-		{
-			if (bIsSCRGB)
-			{
-				PixelShader = TShaderMapRef<FSlatePostProcessUpsamplePS<ESlatePostProcessUpsamplePSPermutation::HDR_SCRGB> >(ShaderMap);
-			}
-			else
-			{
-				PixelShader = TShaderMapRef<FSlatePostProcessUpsamplePS<ESlatePostProcessUpsamplePSPermutation::HDR_PQ10> >(ShaderMap);
-			}
-		}
-		else
-		{
-			PixelShader = TShaderMapRef<FSlatePostProcessUpsamplePS<ESlatePostProcessUpsamplePSPermutation::SDR> >(ShaderMap);
-		}
+		FSlatePostProcessDownsamplePassInputs DownsampleInputs;
+		DownsampleInputs.InputTexture = BlurInputTexture;
+		DownsampleInputs.OutputExtent = OutputTextureViewport.Extent;
 
-		GraphicsPSOInit.BoundShaderState.VertexDeclarationRHI = GFilterVertexDeclaration.VertexDeclarationRHI;
-		GraphicsPSOInit.BoundShaderState.VertexShaderRHI = VertexShader.GetVertexShader();
-		GraphicsPSOInit.BoundShaderState.PixelShaderRHI = PixelShader.GetPixelShader();
-		GraphicsPSOInit.PrimitiveType = PT_TriangleList;
-		SetGraphicsPipelineState(RHICmdList, GraphicsPSOInit, Params.StencilRef);
-
-		const FVector2f SizeUV(
-			DownsampledWidth == SrcTextureWidth ? 1.0f : (DownsampledWidth / (float)SrcTextureWidth) - (1.0f / (float)SrcTextureWidth),
-			DownsampledHeight == SrcTextureHeight ? 1.0f : (DownsampledHeight / (float)SrcTextureHeight) - (1.0f / (float)SrcTextureHeight)
-			);
-
-		const FVector2f Size(DestRect.Right - DestRect.Left, DestRect.Bottom - DestRect.Top);
-		FShaderParams ShaderParams = FShaderParams::MakePixelShaderParams(FVector4f(Size, SizeUV), Params.CornerRadius);
-
-		FRHIBatchedShaderParameters& BatchedParameters = RHICmdList.GetScratchShaderParameters();
-		PixelShader->SetShaderParams(BatchedParameters, ShaderParams);
-		PixelShader->SetTexture(BatchedParameters, SrcTexture, Sampler);
-		RHICmdList.SetBatchedShaderParameters(PixelShader.GetPixelShader(), BatchedParameters);
-
-
-		RendererModule.DrawRectangle(RHICmdList,
-			DestRect.Left, DestRect.Top,
-			Size.X, Size.Y,
-			0, 0,
-			SizeUV.X, SizeUV.Y,
-			DestTexture->GetSizeXY(),
-			FIntPoint(1, 1),
-			VertexShader,
-			EDRF_Default);
+		BlurInputTexture = AddSlatePostProcessDownsamplePass(GraphBuilder, DownsampleInputs);
 	}
-	RHICmdList.EndRenderPass();
 
-}
-#define BILINEAR_FILTER_METHOD 1
+	const ERHIFeatureLevel::Type FeatureLevel = GMaxRHIFeatureLevel;
+	FGlobalShaderMap* ShaderMap = GetGlobalShaderMap(FeatureLevel);
+	TShaderMapRef<FSlatePostProcessBlurPS> PixelShader(ShaderMap);
 
-#if !BILINEAR_FILTER_METHOD
+	FScreenPassRenderTarget BlurOutputTexture(
+		GraphBuilder.CreateTexture(FRDGTextureDesc::Create2D(OutputTextureViewport.Extent, InputPixelFormat, FClearValueBinding::None, GetSlateTransientRenderTargetFlags()), TEXT("SlateBlurHorizontalTexture")),
+		ERenderTargetLoadAction::ENoAction);
 
-static int32 ComputeWeights(int32 KernelSize, float Sigma, TArray<FVector4f>& OutWeightsAndOffsets)
-{
-	OutWeightsAndOffsets.AddUninitialized(KernelSize / 2 + 1);
+	FSlatePostProcessBlurPS::FParameters* PassParameters = nullptr;
 
-	int32 SampleIndex = 0;
-	for (int32 X = 0; X < KernelSize; X += 2)
 	{
-		float Dist = X;
-		FVector4f WeightAndOffset;
-		WeightAndOffset.X = (1.0f / FMath::Sqrt(2 * PI*Sigma*Sigma))*FMath::Exp(-(Dist*Dist) / (2 * Sigma*Sigma));
-		WeightAndOffset.Y = Dist;
+		const FScreenPassTextureViewport BlurInputViewport(BlurInputTexture);
+		const FScreenPassTextureViewportParameters BlurInputParameters = GetScreenPassTextureViewportParameters(BlurInputViewport);
 
-		Dist = X + 1;
-		WeightAndOffset.Z = (1.0f / FMath::Sqrt(2 * PI*Sigma*Sigma))*FMath::Exp(-(Dist*Dist) / (2 * Sigma*Sigma));
-		WeightAndOffset.W = Dist;
+		PassParameters = GraphBuilder.AllocParameters<FSlatePostProcessBlurPS::FParameters>();
+		PassParameters->RenderTargets[0] = BlurOutputTexture.GetRenderTargetBinding();
+		PassParameters->ElementTexture = BlurInputTexture.Texture;
+		PassParameters->ElementTextureSampler = TStaticSamplerState<SF_Bilinear, AM_Clamp, AM_Clamp, AM_Clamp>::GetRHI();
+		PassParameters->SampleCount = SampleCount;
+		PassParameters->BufferSizeAndDirection = FVector4f(BlurInputParameters.ExtentInverse, FVector2f(1.0f, 0.0f));
+		PassParameters->UVBounds = FVector4f(BlurInputParameters.UVViewportBilinearMin, BlurInputParameters.UVViewportBilinearMax);
 
-		OutWeightsAndOffsets[SampleIndex] = WeightAndOffset;
+		check(PassParameters->WeightAndOffsets.Num() * sizeof(PassParameters->WeightAndOffsets[0]) >= WeightsAndOffsets.Num() * sizeof(WeightsAndOffsets[0]));
+		FPlatformMemory::Memcpy(PassParameters->WeightAndOffsets.GetData(), WeightsAndOffsets.GetData(), WeightsAndOffsets.Num() * sizeof(WeightsAndOffsets[0]));
 
-		++SampleIndex;
+		AddDrawScreenPass(GraphBuilder, RDG_EVENT_NAME("Horizontal"), FeatureLevel, FScreenPassTextureViewport(BlurOutputTexture), BlurInputViewport, PixelShader, PassParameters);
 	}
 
-	return KernelSize;
-};
-
-#else
-
-static float GetWeight(float Dist, float Strength)
-{
-	// from https://en.wikipedia.org/wiki/Gaussian_blur
-	float Strength2 = Strength*Strength;
-	return (1.0f / FMath::Sqrt(2 * PI*Strength2))*FMath::Exp(-(Dist*Dist) / (2 * Strength2));
-}
+	BlurInputTexture = BlurOutputTexture;
+	BlurOutputTexture = FScreenPassRenderTarget(
+		GraphBuilder.CreateTexture(FRDGTextureDesc::Create2D(OutputTextureViewport.Extent, InputPixelFormat, FClearValueBinding::None, GetSlateTransientRenderTargetFlags()), TEXT("SlateBlurVerticalTexture")),
+		ERenderTargetLoadAction::ENoAction);
 
-static FVector2f GetWeightAndOffset(float Dist, float Sigma)
-{
-	float Offset1 = Dist;
-	float Weight1 = GetWeight(Offset1, Sigma);
+	{
+		const FScreenPassTextureViewport BlurInputViewport(BlurInputTexture);
+		const FScreenPassTextureViewportParameters BlurInputParameters = GetScreenPassTextureViewportParameters(BlurInputViewport);
 
-	float Offset2 = Dist + 1;
-	float Weight2 = GetWeight(Offset2, Sigma);
+		PassParameters = GraphBuilder.AllocParameters<FSlatePostProcessBlurPS::FParameters>();
+		PassParameters->RenderTargets[0] = BlurOutputTexture.GetRenderTargetBinding();
+		PassParameters->ElementTexture = BlurInputTexture.Texture;
+		PassParameters->ElementTextureSampler = TStaticSamplerState<SF_Bilinear, AM_Clamp, AM_Clamp, AM_Clamp>::GetRHI();
+		PassParameters->SampleCount = SampleCount;
+		PassParameters->BufferSizeAndDirection = FVector4f(BlurInputParameters.ExtentInverse, FVector2f(0.0f, 1.0f));
+		PassParameters->UVBounds = FVector4f(BlurInputParameters.UVViewportBilinearMin, BlurInputParameters.UVViewportBilinearMax);
 
-	float TotalWeight = Weight1 + Weight2;
+		check(PassParameters->WeightAndOffsets.Num() * sizeof(PassParameters->WeightAndOffsets[0]) >= WeightsAndOffsets.Num() * sizeof(WeightsAndOffsets[0]));
+		FPlatformMemory::Memcpy(PassParameters->WeightAndOffsets.GetData(), WeightsAndOffsets.GetData(), WeightsAndOffsets.Num() * sizeof(WeightsAndOffsets[0]));
 
-	float Offset = 0;
-	if (TotalWeight > 0)
-	{
-		Offset = (Weight1*Offset1 + Weight2*Offset2) / TotalWeight;
+		AddDrawScreenPass(GraphBuilder, RDG_EVENT_NAME("Vertical"), FeatureLevel, FScreenPassTextureViewport(BlurOutputTexture), BlurInputViewport, PixelShader, PassParameters);
 	}
 
+	FSlatePostProcessUpsampleInputs UpsampleInputs;
+	UpsampleInputs.InputTexture = BlurOutputTexture;
+	UpsampleInputs.OutputTextureToClear = Inputs.SDRCompositeUITexture;
+	UpsampleInputs.OutputTexture = Inputs.OutputTexture;
+	UpsampleInputs.OutputRect = Inputs.OutputRect;
+	UpsampleInputs.ClippingOp = Inputs.ClippingOp;
+	UpsampleInputs.ClippingStencilBinding = Inputs.ClippingStencilBinding;
+	UpsampleInputs.CornerRadius = Inputs.CornerRadius;
 
-	return FVector2f(TotalWeight, Offset);
+	AddSlatePostProcessUpsamplePass(GraphBuilder, UpsampleInputs);
 }
 
-static int32 ComputeWeights(int32 KernelSize, float Sigma, TArray<FVector4f>& OutWeightsAndOffsets)
+void AddSlatePostProcessBlurPass(FRDGBuilder& GraphBuilder, const FSlatePostProcessSimpleBlurPassInputs& SimpleInputs)
 {
-	int32 NumSamples = FMath::DivideAndRoundUp(KernelSize, 2);
-
-	// We need half of the sample count array because we're packing two samples into one float4
-
-	OutWeightsAndOffsets.AddUninitialized(NumSamples%2 == 0 ? NumSamples / 2 : NumSamples/2+1);
+	const int32 MinKernelSize = 3;
+	const int32 MaxKernelSize = 255;
+	const int32 Downsample2Threshold = 9;
+	const int32 Downsample4Threshold = 64;
+	const float StrengthToKernelSize = 3.0f;
+	const float MinStrength = 0.5f;
+
+	float Strength = FMath::Max(MinStrength, SimpleInputs.Strength);
+	int32 KernelSize = FMath::RoundToInt(Strength * StrengthToKernelSize);
+	int32 DownsampleAmount = 0;
+
+	if (KernelSize > Downsample2Threshold)
+	{
+		DownsampleAmount = KernelSize >= Downsample4Threshold ? 4 : 2;
+		KernelSize /= DownsampleAmount;
+	}
 
-	OutWeightsAndOffsets[0] = FVector4f(FVector2f(GetWeight(0,Sigma), 0), GetWeightAndOffset(1, Sigma) );
-	int32 SampleIndex = 1;
-	for (int32 X = 3; X < KernelSize; X += 4)
+	// Kernel sizes must be odd
+	if (KernelSize % 2 == 0)
 	{
-		OutWeightsAndOffsets[SampleIndex] = FVector4f(GetWeightAndOffset((float)X, Sigma), GetWeightAndOffset((float)(X + 2), Sigma));
+		++KernelSize;
+	}
 
-		++SampleIndex;
+	if (DownsampleAmount > 0)
+	{
+		Strength /= DownsampleAmount;
 	}
 
-	return NumSamples;
-};
+	KernelSize = FMath::Clamp(KernelSize, MinKernelSize, MaxKernelSize);
 
-#endif
+	FSlatePostProcessBlurPassInputs Inputs;
+	Inputs.InputTexture     = SimpleInputs.InputTexture.Texture;
+	Inputs.InputRect        = SimpleInputs.InputTexture.ViewRect;
+	Inputs.OutputTexture    = SimpleInputs.OutputTexture.Texture;
+	Inputs.OutputRect       = SimpleInputs.OutputTexture.ViewRect;
+	Inputs.KernelSize       = KernelSize;
+	Inputs.Strength         = Strength;
+	Inputs.DownsampleAmount = DownsampleAmount;
 
-int32 FSlatePostProcessor::ComputeBlurWeights(int32 KernelSize, float StdDev, TArray<FVector4f>& OutWeightsAndOffsets)
-{
-	return ComputeWeights(KernelSize, StdDev, OutWeightsAndOffsets);
+	AddSlatePostProcessBlurPass(GraphBuilder, Inputs);
 }
 
-void FSlatePostProcessor::TickPostProcessResources()
+//////////////////////////////////////////////////////////////////////////
+
+class FSlatePostProcessColorDeficiencyPS : public FGlobalShader
 {
-	if (GSlateEnableDeleteUnusedPostProcess > 0)
-	{
-		check(IsInRenderingThread());
+public:
+	DECLARE_GLOBAL_SHADER(FSlatePostProcessColorDeficiencyPS);
+	SHADER_USE_PARAMETER_STRUCT(FSlatePostProcessColorDeficiencyPS, FGlobalShader);
+
+	BEGIN_SHADER_PARAMETER_STRUCT(FParameters, )
+		SHADER_PARAMETER_RDG_TEXTURE(Texture2D, ElementTexture)
+		SHADER_PARAMETER_SAMPLER(SamplerState, ElementTextureSampler)
+		SHADER_PARAMETER(float, ColorVisionDeficiencyType)
+		SHADER_PARAMETER(float, ColorVisionDeficiencySeverity)
+		SHADER_PARAMETER(float, bCorrectDeficiency)
+		SHADER_PARAMETER(float, bSimulateCorrectionWithDeficiency)
+		RENDER_TARGET_BINDING_SLOTS()
+	END_SHADER_PARAMETER_STRUCT()
+};
 
-		for (TArray<FSlatePostProcessResource*>::TIterator It = IntermediateTargetsArray.CreateIterator(); It; ++It)
-		{
-			if (GFrameCounter - (*It)->GetFrameUsed() > GSlateEnableDeleteUnusedPostProcess)
-			{
-				(*It)->CleanUp();
-				It.RemoveCurrent();
-			}
-		}
-	}
+IMPLEMENT_GLOBAL_SHADER(FSlatePostProcessColorDeficiencyPS, "/Engine/Private/SlatePostProcessColorDeficiencyPixelShader.usf", "ColorDeficiencyMain", SF_Pixel);
+
+void AddSlatePostProcessColorDeficiencyPass(FRDGBuilder& GraphBuilder, const FSlatePostProcessColorDeficiencyPassInputs& Inputs)
+{
+	const ERHIFeatureLevel::Type FeatureLevel = GMaxRHIFeatureLevel;
+	FGlobalShaderMap* ShaderMap = GetGlobalShaderMap(FeatureLevel);
+	TShaderMapRef<FSlatePostProcessColorDeficiencyPS> PixelShader(ShaderMap);
+	const FRDGTextureDesc& InputDesc = Inputs.InputTexture.Texture->Desc;
+
+	const FScreenPassRenderTarget Output(
+		GraphBuilder.CreateTexture(FRDGTextureDesc::Create2D(InputDesc.Extent, InputDesc.Format, FClearValueBinding::None, GetSlateTransientRenderTargetFlags()), TEXT("ColorDeficiency")),
+		ERenderTargetLoadAction::ENoAction);
+
+	FSlatePostProcessColorDeficiencyPS::FParameters* PassParameters = GraphBuilder.AllocParameters<FSlatePostProcessColorDeficiencyPS::FParameters>();
+	PassParameters->RenderTargets[0] = Output.GetRenderTargetBinding();
+	PassParameters->ElementTexture = Inputs.InputTexture.Texture;
+	PassParameters->ElementTextureSampler = TStaticSamplerState<SF_Point, AM_Clamp, AM_Clamp, AM_Clamp>::GetRHI();
+	PassParameters->ColorVisionDeficiencyType = (float)GSlateColorDeficiencyType;
+	PassParameters->ColorVisionDeficiencySeverity = (float)GSlateColorDeficiencySeverity;
+	PassParameters->bCorrectDeficiency = GSlateColorDeficiencyCorrection ? 1.0f : 0.0f;
+	PassParameters->bSimulateCorrectionWithDeficiency = GSlateShowColorDeficiencyCorrectionWithDeficiency ? 1.0f : 0.0f;
+
+	const FScreenPassTextureViewport Viewport(Output);
+	AddDrawScreenPass(GraphBuilder, RDG_EVENT_NAME("ColorDeficiency"), FeatureLevel, Viewport, Viewport, PixelShader, PassParameters);
+
+	FSlatePostProcessUpsampleInputs UpsampleInputs;
+	UpsampleInputs.InputTexture = Output;
+	UpsampleInputs.OutputTexture = Inputs.OutputTexture.Texture;
+	UpsampleInputs.OutputRect = Inputs.OutputTexture.ViewRect;
+
+	AddSlatePostProcessUpsamplePass(GraphBuilder, UpsampleInputs);
 }
\ No newline at end of file
diff --git a/Engine/Source/Runtime/SlateRHIRenderer/Private/SlatePostProcessor.h b/Engine/Source/Runtime/SlateRHIRenderer/Private/SlatePostProcessor.h
--- a/Engine/Source/Runtime/SlateRHIRenderer/Private/SlatePostProcessor.h
+++ b/Engine/Source/Runtime/SlateRHIRenderer/Private/SlatePostProcessor.h
@@ -2,61 +2,82 @@
 
 #pragma once
 
-#include "CoreMinimal.h"
-#include "RHI.h"
 #include "Layout/SlateRect.h"
+#include "Layout/Clipping.h"
+#include "ScreenPass.h"
 
-class FSlatePostProcessResource;
-class IRendererModule;
-struct IPooledRenderTarget;
-
-enum class EPostProcessDestination : uint8
+struct FSlateClippingOp
 {
-	// Output postprocess to the default UI render target
-	UITarget,
-	// Output postprocess to the provided destination texture
-	DestTexture
+	union
+	{
+		struct
+		{
+			FSlateRect Rect;
+		} Data_Scissor;
+
+		struct
+		{
+			TConstArrayView<FSlateClippingZone> Zones;
+		} Data_Stencil;
+	};
+
+	FVector2f Offset;
+	EClippingMethod Method;
+	uint8 MaskingId;
+
+	static inline FSlateClippingOp* Scissor(FRDGBuilder& GraphBuilder, FVector2f Offset, FSlateRect Rect)
+	{
+		FSlateClippingOp* Op = GraphBuilder.AllocPOD<FSlateClippingOp>();
+		Op->Data_Scissor.Rect = Rect;
+		Op->Offset = Offset;
+		Op->Method = EClippingMethod::Scissor;
+		Op->MaskingId = 0;
+		return Op;
+	}
+
+	static inline FSlateClippingOp* Stencil(FRDGBuilder& GraphBuilder, FVector2f Offset, TConstArrayView<FSlateClippingZone> Zones, int32 MaskingId)
+	{
+		FSlateClippingOp* Op = GraphBuilder.AllocPOD<FSlateClippingOp>();
+		Op->Data_Stencil.Zones = Zones;
+		Op->Offset = Offset;
+		Op->Method = EClippingMethod::Stencil;
+		Op->MaskingId = MaskingId;
+		return Op;
+	}
 };
 
-struct FPostProcessRectParams
+bool GetSlateClippingPipelineState(const FSlateClippingOp* ClippingStateOp, FRHIDepthStencilState*& OutDepthStencilState, uint8& OutStencilRef);
+
+void SetSlateClipping(FRHICommandList& RHICmdList, const FSlateClippingOp* ClippingStateOp, FIntRect ViewportRect);
+
+struct FSlatePostProcessBlurPassInputs
 {
-	FTextureRHIRef SourceTexture;
-	FTextureRHIRef DestTexture; // Only used when 'PostProcessDest' is 'DestTexture'
-	FSlateRect SourceRect;
-	FSlateRect DestRect;
-	FVector4f CornerRadius;
-	FIntPoint SourceTextureSize;
-	TFunction<void(FRHICommandListImmediate&, FGraphicsPipelineStateInitializer&, FRHIRenderPassInfo&)> RestoreStateFunc;
-	TRefCountPtr<IPooledRenderTarget> UITarget; // not using FTextureRHIRef because we want to be able to use FRenderTargetWriteMask::Decode
-	uint32 StencilRef{};
-	EDisplayColorGamut HDRDisplayColorGamut;
-	EPostProcessDestination PostProcessDest = EPostProcessDestination::UITarget;
+	// An optional in/out separately composited UI texture that is composited with the input and then the output rect is reset to transparent.
+	FRDGTexture* SDRCompositeUITexture = nullptr;
+	FRDGTexture* InputTexture = nullptr;
+	FRDGTexture* OutputTexture = nullptr;
+	ERenderTargetLoadAction OutputLoadAction = ERenderTargetLoadAction::ELoad;
+
+	const FSlateClippingOp* ClippingOp = nullptr;
+	const FDepthStencilBinding* ClippingStencilBinding = nullptr;
+
+	FIntRect InputRect;
+	FIntRect OutputRect;
+	uint32 KernelSize = 0;
+	float  Strength = 0.0f;
+	uint32 DownsampleAmount = 0;
+	FVector4f CornerRadius = FVector4f::Zero();
 };
 
-struct FBlurRectParams
+extern ETextureCreateFlags GetSlateTransientRenderTargetFlags();
+extern ETextureCreateFlags GetSlateTransientDepthStencilFlags();
+
+void AddSlatePostProcessBlurPass(FRDGBuilder& GraphBuilder, const FSlatePostProcessBlurPassInputs& Inputs);
+
+struct FSlatePostProcessColorDeficiencyPassInputs
 {
-	int32 KernelSize;
-	int32 DownsampleAmount;
-	float Strength;
+	FScreenPassTexture InputTexture;
+	FScreenPassTexture OutputTexture;
 };
 
-class FSlatePostProcessor
-{
-public:
-	FSlatePostProcessor();
-	~FSlatePostProcessor();
-
-	void BlurRect(FRHICommandListImmediate& RHICmdList, IRendererModule& RendererModule, const FBlurRectParams& Params, const FPostProcessRectParams& RectParams);
-	
-	void ColorDeficiency(FRHICommandListImmediate& RHICmdList, IRendererModule& RendererModule, const FPostProcessRectParams& RectParams);
-	
-	void ReleaseRenderTargets();
-	void TickPostProcessResources();
-
-private:
-	void DownsampleRect(FRHICommandListImmediate& RHICmdList, IRendererModule& RendererModule, const FPostProcessRectParams& Params, const FIntPoint& DownsampleSize, FSlatePostProcessResource* IntermediateTargets);
-	void UpsampleRect(FRHICommandListImmediate& RHICmdList, IRendererModule& RendererModule, const FPostProcessRectParams& Params, const FIntPoint& DownsampleSize, FSamplerStateRHIRef& Sampler, FSlatePostProcessResource* IntermediateTargets);
-	int32 ComputeBlurWeights(int32 KernelSize, float StdDev, TArray<FVector4f>& OutWeightsAndOffsets);
-private:
-	TArray<FSlatePostProcessResource*> IntermediateTargetsArray;
-};
\ No newline at end of file
+void AddSlatePostProcessColorDeficiencyPass(FRDGBuilder& GraphBuilder, const FSlatePostProcessColorDeficiencyPassInputs& Inputs);
\ No newline at end of file
