Clean-up of Slate RHI renderer and conversion to RDG.

diff --git a/Engine/Source/Runtime/SlateRHIRenderer/Private/SlateRHIRenderer.cpp b/Engine/Source/Runtime/SlateRHIRenderer/Private/SlateRHIRenderer.cpp
--- a/Engine/Source/Runtime/SlateRHIRenderer/Private/SlateRHIRenderer.cpp
+++ b/Engine/Source/Runtime/SlateRHIRenderer/Private/SlateRHIRenderer.cpp
@@ -53,19 +53,15 @@
 #include "HDRHelper.h"
 #include "RenderCore.h"
 #include "DataDrivenShaderPlatformInfo.h"
+#include "SlatePostProcessor.h"
 
 #if WITH_EDITORONLY_DATA
 #include "ShaderCompiler.h"
 #endif
 
-DECLARE_CYCLE_STAT(TEXT("Slate RT: Rendering"), STAT_SlateRenderingRTTime, STATGROUP_Slate);
-
-DECLARE_CYCLE_STAT(TEXT("Slate RT: Draw Batches"), STAT_SlateRTDrawBatches, STATGROUP_Slate);
-
 DECLARE_CYCLE_STAT(TEXT("Total Render Thread time including dependent waits"), STAT_RenderThreadCriticalPath, STATGROUP_Threading);
 
 CSV_DEFINE_CATEGORY(RenderThreadIdle, true);
-
 CSV_DECLARE_CATEGORY_MODULE_EXTERN(SLATECORE_API, Slate);
 
 DECLARE_GPU_DRAWCALL_STAT_NAMED(SlateUI, TEXT("Slate UI"));
@@ -94,12 +90,6 @@ static TAutoConsoleVariable<int32> CVarUICompositeMode(
 	TEXT("1: Shader pass to improve HDR blending\n"),
 	ECVF_RenderThreadSafe);
 
-static TAutoConsoleVariable<float> CVarDrawToVRRenderTarget(
-	TEXT("Slate.DrawToVRRenderTarget"),
-	1,
-	TEXT("If enabled while in VR. Slate UI will be drawn into the render target texture where the VR imagery for either eye was rendered, allow the viewer of the HMD to see the UI (for better or worse.)  This render target will then be cropped/scaled into the back buffer, if mirroring is enabled.  When disabled, Slate UI will be drawn on top of the backbuffer (not to the HMD) after the mirror texture has been cropped/scaled into the backbuffer."),
-	ECVF_RenderThreadSafe);
-
 static TAutoConsoleVariable<int32> CVarCopyBackbufferToSlatePostRenderTargets(
 	TEXT("Slate.CopyBackbufferToSlatePostRenderTargets"),
 	0,
@@ -112,138 +102,107 @@ TAutoConsoleVariable<int32> CVarShowSlateOverdraw(
 	TEXT("Slate.ShowOverdraw"),
 	0,
 	TEXT("0: Don't show overdraw, 1: Show Overdraw"),
-	ECVF_Default
+	ECVF_RenderThreadSafe
 );
 
 TAutoConsoleVariable<int32> CVarShowSlateBatching(
 	TEXT("Slate.ShowBatching"),
 	0,
 	TEXT("0: Don't show batching, 1: Show Batching"),
-	ECVF_Default
+	ECVF_RenderThreadSafe
 );
 #endif
 
+bool GSlateWireframe = false;
+static FAutoConsoleVariableRef CVarSlateWireframe(TEXT("Slate.ShowWireFrame"), GSlateWireframe, TEXT(""), ECVF_Default);
+
 // RT stat including waits toggle. Off by default for historical tracking reasons
-TAutoConsoleVariable<int32> CVarRenderThreadTimeIncludesDependentWaits(
+static TAutoConsoleVariable<int32> CVarRenderThreadTimeIncludesDependentWaits(
 	TEXT("r.RenderThreadTimeIncludesDependentWaits"),
 	0,
 	TEXT("0: RT stat only includes non-idle time, 1: RT stat includes dependent waits (matching RenderThreadTime_CriticalPath)"),
 	ECVF_Default
 );
 
-
-struct FSlateDrawWindowCommandParams
+static bool IsVSyncRequired(const FSlateElementBatcher& ElementBatcher)
 {
-	FSlateRHIRenderer* Renderer;
-	FSlateWindowElementList* WindowElementList;
-	SWindow* Window;
-	FIntPoint CursorPostion;
-	FIntRect ViewRect;
-	float ViewportScaleUI;
-	ESlatePostRT UsedSlatePostBuffers;
-#if WANTS_DRAW_MESH_EVENTS
-	FString WindowTitle;
-#endif
-	FGameTime Time;
-	bool bLockToVsync;
-	bool bClear;
-};
-
-void FViewportInfo::InitRHI(FRHICommandListBase&)
-{
-	// Viewport RHI is created on the game thread
-	// Create the depth-stencil surface if needed.
-	RecreateDepthBuffer_RenderThread();
-}
-
-void FViewportInfo::ReleaseRHI()
-{
-	// Some RHIs delete resources when the command list completes on the GPU, but they don't take Present calls into account due to
-	// API limitations. Sync the GPU here, to make sure nothing is using the backbuffer anymore.
-	FRHICommandListExecutor::GetImmediateCommandList().BlockUntilGPUIdle();
-	DepthStencil.SafeRelease();
-	ViewportRHI.SafeRelease();
-}
-
-void FViewportInfo::ReleaseResource()
-{
-	FRenderResource::ReleaseResource();
-	UITargetRT.SafeRelease();
-}
+	bool bLockToVsync = ElementBatcher.RequiresVsync();
 
-void FViewportInfo::ConditionallyUpdateDepthBuffer(bool bInRequiresStencilTest, uint32 InWidth, uint32 InHeight)
-{
-	check(IsInRenderingThread());
-
-	bool bDepthStencilStale =
-		bInRequiresStencilTest &&
-		(!bRequiresStencilTest ||
-		(DepthStencil.IsValid() && (DepthStencil->GetSizeX() != InWidth || DepthStencil->GetSizeY() != InHeight)));
-
-	bRequiresStencilTest = bInRequiresStencilTest;
-
-	// Allocate a stencil buffer if needed and not already allocated
-	if (bDepthStencilStale)
+	if (GIsEditor)
 	{
-		RecreateDepthBuffer_RenderThread();
+		static IConsoleVariable* CVar = IConsoleManager::Get().FindConsoleVariable(TEXT("r.VSyncEditor"));
+		bLockToVsync |= (CVar->GetInt() != 0);
 	}
-}
-
-void FViewportInfo::RecreateDepthBuffer_RenderThread()
-{
-	check(IsInRenderingThread());
-	DepthStencil.SafeRelease();
-	if (bRequiresStencilTest)
+	else
 	{
-		ETextureCreateFlags TargetableTextureFlags = TexCreate_DepthStencilTargetable;
-	
-		const FRHITextureCreateDesc Desc =
-			FRHITextureCreateDesc::Create2D(TEXT("SlateViewportDepthStencil"))
-			.SetExtent(Width, Height)
-			.SetFormat(PF_DepthStencil)
-			.SetFlags(TargetableTextureFlags | ETextureCreateFlags::ShaderResource)
-			.SetInitialState(ERHIAccess::SRVMask)
-			.SetClearValue(FClearValueBinding::DepthZero);
-
-		DepthStencil = RHICreateTexture(Desc);
-		check(IsValidRef(DepthStencil));
+		static IConsoleVariable* CVar = IConsoleManager::Get().FindConsoleVariable(TEXT("r.VSync"));
+		bLockToVsync |= (CVar->GetInt() != 0);
 	}
+
+	return bLockToVsync;
 }
 
-FSlateRHIRenderer::FSlateRHIRenderer(TSharedRef<FSlateFontServices> InSlateFontServices, TSharedRef<FSlateRHIResourceManager> InResourceManager)
-	: FSlateRenderer(InSlateFontServices)
-	, EnqueuedWindowDrawBuffer(NULL)
-	, FreeBufferIndex(0)
-	, FastPathRenderingDataCleanupList(nullptr)
-	, bUpdateHDRDisplayInformation(false)
-	, CurrentSceneIndex(-1)
-	, ResourceVersion(0)
+struct FSlateViewportInfo : public FRenderResource
 {
-	ResourceManager = InResourceManager;
-
-	ViewMatrix = FMatrix(FPlane(1, 0, 0, 0),
-		FPlane(0, 1, 0, 0),
-		FPlane(0, 0, 1, 0),
-		FPlane(0, 0, 0, 1));
+	FViewportRHIRef ViewportRHI;
+	void* OSWindow = nullptr;
+	IViewportRenderTargetProvider* RenderTargetProvider = nullptr;
+	FMatrix ProjectionMatrix;
+	FIntPoint Extent = FIntPoint::ZeroValue;
+	FIntPoint ExtentToResizeTo = FIntPoint::ZeroValue;
+	EPixelFormat PixelFormat = EPixelFormat::PF_Unknown;
+	EDisplayColorGamut HDRDisplayColorGamut = EDisplayColorGamut::sRGB_D65;
+	EDisplayOutputFormat HDRDisplayOutputFormat = EDisplayOutputFormat::SDR_sRGB;
+	bool bDisplayFormatIsHDR = false;
+	bool bFullscreen = false;
 
-	for (uint64& LastFramePostBufferUsed : LastFramesPostBufferUsed)
+	void ReleaseRHI() override
 	{
-		LastFramePostBufferUsed = 0;
+		// Full GPU sync here to simplify memory lifetime of the underlying resource.
+		FRHICommandListExecutor::GetImmediateCommandList().BlockUntilGPUIdle();
+		ViewportRHI.SafeRelease();
 	}
 
-	bTakingAScreenShot = false;
-	OutScreenshotData = NULL;
-	OutHDRScreenshotData = NULL;
-	ScreenshotViewportInfo = nullptr;
-	bIsStandaloneStereoOnlyDevice = IHeadMountedDisplayModule::IsAvailable() && IHeadMountedDisplayModule::Get().IsStandaloneStereoOnlyDevice();
-	bShrinkPostBufferRequested = ESlatePostRT::None;
-}
+	FTextureRHIRef GetRenderTargetTexture() const
+	{
+		if (RenderTargetProvider)
+		{
+			if (FSlateShaderResource* RenderTargetTexture = RenderTargetProvider->GetViewportRenderTargetTexture())
+			{
+				return static_cast<FSlateRenderTargetRHI*>(RenderTargetTexture)->GetTypedResource();
+			}
+		}
+		return nullptr;
+	}
+};
 
-FSlateRHIRenderer::~FSlateRHIRenderer()
+struct FSlatePostProcessUpdateRequest
 {
-}
+	TSharedPtr<FSlateRHIPostBufferProcessorProxy> PostProcessorProxy;
+	FTextureResource* RenderTargetTextureResource = nullptr;
+	ESlatePostRT RenderTarget = ESlatePostRT::None;
+};
 
-FMatrix FSlateRHIRenderer::CreateProjectionMatrix(uint32 Width, uint32 Height)
+struct FSlateDrawWindowPassInputs
+{
+	FSlateRHIRenderer* Renderer = nullptr;
+	FSlateWindowElementList* WindowElementList = nullptr;
+	SWindow* Window = nullptr;
+	FSlateViewportInfo* ViewportInfo = nullptr;
+	TConstArrayView<FSlatePostProcessUpdateRequest> PostProcessUpdateRequests;
+	FIntPoint CursorPosition = FIntPoint::ZeroValue;
+	FIntRect SceneViewRect;
+	float ViewportScaleUI = 0.0f;
+	ESlatePostRT UsedSlatePostBuffers = ESlatePostRT::None;
+#if WANTS_DRAW_MESH_EVENTS
+	FString WindowTitle;
+#endif
+	FGameTime Time;
+	bool bLockToVsync = false;
+	bool bClear = false;
+};
+
+FMatrix CreateSlateProjectionMatrix(uint32 Width, uint32 Height)
 {
 	// Create ortho projection matrix
 	const float Left = 0;
@@ -262,14 +221,15 @@ FMatrix FSlateRHIRenderer::CreateProjectionMatrix(uint32 Width, uint32 Height)
 	);
 }
 
-int32 FSlateRHIRenderer::GetDrawToVRRenderTarget()
-{
-	return CVarDrawToVRRenderTarget->GetInt();
-}
-
-int32 FSlateRHIRenderer::GetProcessSlatePostBuffers()
+FSlateRHIRenderer::FSlateRHIRenderer(TSharedRef<FSlateFontServices> InSlateFontServices, TSharedRef<FSlateRHIResourceManager> InResourceManager)
+	: FSlateRenderer(InSlateFontServices)
+	, ResourceManager(InResourceManager)
+	, bIsStandaloneStereoOnlyDevice(IHeadMountedDisplayModule::IsAvailable() && IHeadMountedDisplayModule::Get().IsStandaloneStereoOnlyDevice())
 {
-	return CVarCopyBackbufferToSlatePostRenderTargets->GetInt();
+	for (uint64& LastFramePostBufferUsed : LastFramesPostBufferUsed)
+	{
+		LastFramePostBufferUsed = 0;
+	}
 }
 
 bool FSlateRHIRenderer::Initialize()
@@ -287,21 +247,15 @@ bool FSlateRHIRenderer::Initialize()
 
 void FSlateRHIRenderer::Destroy()
 {
-	RenderingPolicy->ReleaseResources();
 	ResourceManager->ReleaseResources();
 	SlateFontServices->ReleaseResources();
 
-	for (TMap< const SWindow*, FViewportInfo*>::TIterator It(WindowToViewportInfo); It; ++It)
-	{
-		BeginReleaseResource(It.Value());
-	}
-
-	if (FastPathRenderingDataCleanupList)
+	for (auto& Entry : WindowToViewportInfo)
 	{
-		FastPathRenderingDataCleanupList->Cleanup();
-		FastPathRenderingDataCleanupList = nullptr;
+		BeginReleaseResource(Entry.Value);
 	}
 
+	FlushPendingDeletes();
 	FlushRenderingCommands();
 
 	ElementBatcher.Reset();
@@ -311,13 +265,12 @@ void FSlateRHIRenderer::Destroy()
 
 	DeferredUpdateContexts.Empty();
 
-	for (TMap< const SWindow*, FViewportInfo*>::TIterator It(WindowToViewportInfo); It; ++It)
+	for (auto& Entry : WindowToViewportInfo)
 	{
-		FViewportInfo* ViewportInfo = It.Value();
-		delete ViewportInfo;
+		delete Entry.Value;
 	}
-
 	WindowToViewportInfo.Empty();
+
 	CurrentSceneIndex = -1;
 	ActiveScenes.Empty();
 }
@@ -359,13 +312,13 @@ FSlateDrawBuffer& FSlateRHIRenderer::AcquireDrawBuffer()
 	return *Buffer;
 }
 
-void FSlateRHIRenderer::ReleaseDrawBuffer(FSlateDrawBuffer& InWindowDrawBuffer)
+void FSlateRHIRenderer::ReleaseDrawBuffer(FSlateDrawBuffer& WindowDrawBuffer)
 {
 #if DO_CHECK
 	bool bFound = false;
 	for (int32 Index = 0; Index < NumDrawBuffers; ++Index)
 	{
-		if (&DrawBuffers[Index] == &InWindowDrawBuffer)
+		if (&DrawBuffers[Index] == &WindowDrawBuffer)
 		{
 			bFound = true;
 			break;
@@ -374,107 +327,71 @@ void FSlateRHIRenderer::ReleaseDrawBuffer(FSlateDrawBuffer& InWindowDrawBuffer)
 	ensureMsgf(bFound, TEXT("It release a DrawBuffer that is not a member of the SlateRHIRenderer"));
 #endif
 
-	FSlateDrawBuffer* DrawBuffer = &InWindowDrawBuffer;
-	ENQUEUE_RENDER_COMMAND(SlateReleaseDrawBufferCommand)(
-		[DrawBuffer](FRHICommandListImmediate& RHICmdList)
-		{
-			FSlateReleaseDrawBufferCommand::ReleaseDrawBuffer(RHICmdList, DrawBuffer);
-		}
-	);
+	ENQUEUE_RENDER_COMMAND(SlateReleaseDrawBufferCommand)([&WindowDrawBuffer](FRHICommandList& RHICmdList)
+	{
+		WindowDrawBuffer.Unlock();
+	});
 }
 
 void FSlateRHIRenderer::CreateViewport(const TSharedRef<SWindow> Window)
 {
-	FlushRenderingCommands();
-
-	if (!WindowToViewportInfo.Contains(&Window.Get()))
+	if (WindowToViewportInfo.Contains(&Window.Get()))
 	{
-		const FVector2f WindowSize = UE::Slate::CastToVector2f(Window->GetViewportSize());
-
-		// Clamp the window size to a reasonable default anything below 8 is a d3d warning and 8 is used anyway.
-		// @todo Slate: This is a hack to work around menus being summoned with 0,0 for window size until they are ticked.
-		int32 Width = FMath::Max(MIN_VIEWPORT_SIZE, FMath::CeilToInt(WindowSize.X));
-		int32 Height = FMath::Max(MIN_VIEWPORT_SIZE, FMath::CeilToInt(WindowSize.Y));
-
-		// Sanity check dimensions
-		if (!ensureMsgf(Width <= MAX_VIEWPORT_SIZE && Height <= MAX_VIEWPORT_SIZE, TEXT("Invalid window with Width=%u and Height=%u"), Width, Height))
-		{
-			Width = FMath::Clamp(Width, MIN_VIEWPORT_SIZE, MAX_VIEWPORT_SIZE);
-			Height = FMath::Clamp(Height, MIN_VIEWPORT_SIZE, MAX_VIEWPORT_SIZE);
-		}
-
-
-		FViewportInfo* NewInfo = new FViewportInfo();
-		// Create Viewport RHI if it doesn't exist (this must be done on the game thread)
-		TSharedRef<FGenericWindow> NativeWindow = Window->GetNativeWindow().ToSharedRef();
-		NewInfo->OSWindow = NativeWindow->GetOSWindowHandle();
-		NewInfo->Width = Width;
-		NewInfo->Height = Height;
-		NewInfo->DesiredWidth = Width;
-		NewInfo->DesiredHeight = Height;
-		NewInfo->ProjectionMatrix = CreateProjectionMatrix( Width, Height );
-		// In MobileLDR case backbuffer format should match or be compatible with a SceneColor format in FSceneRenderTargets::GetDesiredMobileSceneColorFormat()
-		if (bIsStandaloneStereoOnlyDevice || (GMaxRHIFeatureLevel == ERHIFeatureLevel::ES3_1 && !IsMobileHDR()))
-		{
-			NewInfo->PixelFormat = GetSlateRecommendedColorFormat();
-		}
-#if ALPHA_BLENDED_WINDOWS		
-		if (Window->GetTransparencySupport() == EWindowTransparency::PerPixel)
-		{
-			NewInfo->PixelFormat = GetSlateRecommendedColorFormat();
-		}
-#endif
+		return;
+	}
 
-		// SDR format holds the requested format in non HDR mode
-		NewInfo->SDRPixelFormat = NewInfo->PixelFormat;
-		HDRGetMetaData(NewInfo->HDRDisplayOutputFormat, NewInfo->HDRDisplayColorGamut, NewInfo->bSceneHDREnabled, Window->GetPositionInScreen(), Window->GetPositionInScreen() + Window->GetSizeInScreen(), NewInfo->OSWindow);
+	FlushRenderingCommands();
 
-		if (NewInfo->bSceneHDREnabled)
-		{
-			NewInfo->PixelFormat = GRHIHDRDisplayOutputFormat;
-		}
+	const FVector2f ViewportSize = Window->GetViewportSize();
 
-		// Sanity check dimensions
-		checkf(Width <= MAX_VIEWPORT_SIZE && Height <= MAX_VIEWPORT_SIZE, TEXT("Invalid window with Width=%u and Height=%u"), Width, Height);
+	FIntPoint ExtentToResizeTo;
+	ExtentToResizeTo.X = FMath::CeilToInt(ViewportSize.X);
+	ExtentToResizeTo.Y = FMath::CeilToInt(ViewportSize.Y);
+	ExtentToResizeTo = ExtentToResizeTo.ComponentMax(FIntPoint(MIN_VIEWPORT_SIZE));
 
-		bool bFullscreen = IsViewportFullscreen( *Window );
-		NewInfo->ViewportRHI = RHICreateViewport( NewInfo->OSWindow, Width, Height, bFullscreen, NewInfo->PixelFormat );
-		NewInfo->bFullscreen = bFullscreen;
+	if (!ensureMsgf(ExtentToResizeTo.X <= MAX_VIEWPORT_SIZE && ExtentToResizeTo.Y <= MAX_VIEWPORT_SIZE, TEXT("Invalid window with Width=%u and Height=%u"), ExtentToResizeTo.X, ExtentToResizeTo.Y))
+	{
+		ExtentToResizeTo = ExtentToResizeTo.ComponentMin(FIntPoint(MAX_VIEWPORT_SIZE));
+	}
 
-		// Was the window created on a HDR compatible display?
-		NewInfo->bHDREnabled = RHIGetColorSpace(NewInfo->ViewportRHI) != EColorSpaceAndEOTF::ERec709_sRGB ;
-		Window->SetIsHDR(NewInfo->bHDREnabled);
+	FSlateViewportInfo* ViewInfo = new FSlateViewportInfo();
+	ViewInfo->OSWindow = Window->GetNativeWindow()->GetOSWindowHandle();
+	ViewInfo->ProjectionMatrix = CreateSlateProjectionMatrix(ExtentToResizeTo.X, ExtentToResizeTo.Y);
+	ViewInfo->Extent = ExtentToResizeTo;
+	ViewInfo->ExtentToResizeTo = ExtentToResizeTo;
+	HDRGetMetaData(ViewInfo->HDRDisplayOutputFormat, ViewInfo->HDRDisplayColorGamut, ViewInfo->bDisplayFormatIsHDR, Window->GetPositionInScreen(), Window->GetPositionInScreen() + Window->GetSizeInScreen(), ViewInfo->OSWindow);
 
-		WindowToViewportInfo.Add(&Window.Get(), NewInfo);
+	const bool bFullscreen = IsViewportFullscreen(*Window);
+	ViewInfo->PixelFormat = GetViewportPixelFormat(*Window, ViewInfo->bDisplayFormatIsHDR);
+	ViewInfo->ViewportRHI = RHICreateViewport(ViewInfo->OSWindow, ExtentToResizeTo.X, ExtentToResizeTo.Y, bFullscreen, ViewInfo->PixelFormat);
+	ViewInfo->bFullscreen = bFullscreen;
+	BeginInitResource(ViewInfo);
 
-		BeginInitResource(NewInfo);
-	}
+	WindowToViewportInfo.Add(&Window.Get(), ViewInfo);
+	Window->SetIsHDR(ViewInfo->bDisplayFormatIsHDR);
 }
 
-void FSlateRHIRenderer::ConditionalResizeViewport(FViewportInfo* ViewInfo, uint32 Width, uint32 Height, bool bFullscreen, SWindow* Window)
+void FSlateRHIRenderer::ResizeViewportIfNeeded(FSlateViewportInfo* ViewInfo, FIntPoint ExtentToResizeTo, bool bFullscreen, SWindow* Window)
 {
 	checkSlow(IsThreadSafeForSlateRendering());
 
-	// Force update if HDR output state changes
+	if (!IsInGameThread() || IsInSlateThread() || !ViewInfo)
+	{
+		return;
+	}
 
 	bool bHDREnabled = IsHDREnabled();
 	EDisplayColorGamut HDRColorGamut = HDRGetDefaultDisplayColorGamut();
 	EDisplayOutputFormat HDROutputDevice = HDRGetDefaultDisplayOutputFormat();
+	HDRGetMetaData(HDROutputDevice, HDRColorGamut, bHDREnabled, Window->GetPositionInScreen(), Window->GetPositionInScreen() + Window->GetSizeInScreen(), ViewInfo->OSWindow);
 
 	bool bHDRStale = false;
-	if (ViewInfo)
-	{
-		HDRGetMetaData(HDROutputDevice, HDRColorGamut, bHDREnabled, Window->GetPositionInScreen(), Window->GetPositionInScreen() + Window->GetSizeInScreen(), ViewInfo->OSWindow);
-		bHDRStale |= HDROutputDevice != ViewInfo->HDRDisplayOutputFormat;
-		bHDRStale |= HDRColorGamut != ViewInfo->HDRDisplayColorGamut;
-		bHDRStale |= bHDREnabled != ViewInfo->bSceneHDREnabled;
-	}
+	bHDRStale |= HDROutputDevice != ViewInfo->HDRDisplayOutputFormat;
+	bHDRStale |= HDRColorGamut   != ViewInfo->HDRDisplayColorGamut;
+	bHDRStale |= bHDREnabled     != ViewInfo->bDisplayFormatIsHDR;
 
-	if (IsInGameThread() && !IsInSlateThread() && ViewInfo && (bHDRStale || ViewInfo->Height != Height || ViewInfo->Width != Width || ViewInfo->bFullscreen != bFullscreen || !IsValidRef(ViewInfo->ViewportRHI)))
+	if (bHDRStale || ViewInfo->Extent != ExtentToResizeTo || ViewInfo->bFullscreen != bFullscreen || !IsValidRef(ViewInfo->ViewportRHI))
 	{
-		// The viewport size we have doesn't match the requested size of the viewport.
-		// Resize it now.
-
 		// Prevent the texture update logic to use the RHI while the viewport is resized. 
 		// This could happen if a streaming IO request completes and throws a callback.
 		// @todo : this does not in fact stop texture tasks from using the RHI while the viewport is resized
@@ -485,45 +402,39 @@ void FSlateRHIRenderer::ConditionalResizeViewport(FViewportInfo* ViewInfo, uint3
 		// cannot resize the viewport while potentially using it.
 		FlushRenderingCommands();
 
-		// Windows are allowed to be zero sized ( sometimes they are animating to/from zero for example)
-		// but viewports cannot be zero sized.  Use 8x8 as a reasonably sized viewport in this case.
-		uint32 NewWidth = FMath::Max<uint32>(8, Width);
-		uint32 NewHeight = FMath::Max<uint32>(8, Height);
+		// Windows are allowed to be zero sized (sometimes they are animating to/from zero for example) but not viewports.
+		ExtentToResizeTo = ExtentToResizeTo.ComponentMax(FIntPoint(MIN_VIEWPORT_SIZE, MIN_VIEWPORT_SIZE));
 
-		// Sanity check dimensions
-		if (NewWidth > MAX_VIEWPORT_SIZE)
+		if (ExtentToResizeTo.X > MAX_VIEWPORT_SIZE)
 		{
-			UE_LOG(LogSlate, Warning, TEXT("Tried to set viewport width size to %d.  Clamping size to max allowed size of %d instead."), NewWidth, MAX_VIEWPORT_SIZE);
-			NewWidth = MAX_VIEWPORT_SIZE;
+			UE_LOG(LogSlate, Warning, TEXT("Tried to set viewport width size to %d.  Clamping size to max allowed size of %d instead."), ExtentToResizeTo.X, MAX_VIEWPORT_SIZE);
+			ExtentToResizeTo.X = MAX_VIEWPORT_SIZE;
 		}
 
-		if (NewHeight > MAX_VIEWPORT_SIZE)
+		if (ExtentToResizeTo.Y > MAX_VIEWPORT_SIZE)
 		{
-			UE_LOG(LogSlate, Warning, TEXT("Tried to set viewport height size to %d.  Clamping size to max allowed size of %d instead."), NewHeight, MAX_VIEWPORT_SIZE);
-			NewHeight = MAX_VIEWPORT_SIZE;
+			UE_LOG(LogSlate, Warning, TEXT("Tried to set viewport height size to %d.  Clamping size to max allowed size of %d instead."), ExtentToResizeTo.Y, MAX_VIEWPORT_SIZE);
+			ExtentToResizeTo.Y = MAX_VIEWPORT_SIZE;
 		}
 
-		ViewInfo->Width = NewWidth;
-		ViewInfo->Height = NewHeight;
-		ViewInfo->DesiredWidth = NewWidth;
-		ViewInfo->DesiredHeight = NewHeight;
-		ViewInfo->ProjectionMatrix = CreateProjectionMatrix(NewWidth, NewHeight);
+		ViewInfo->ProjectionMatrix = CreateSlateProjectionMatrix(ExtentToResizeTo.X, ExtentToResizeTo.Y);
+		ViewInfo->Extent = ExtentToResizeTo;
+		ViewInfo->ExtentToResizeTo = ExtentToResizeTo;
 		ViewInfo->bFullscreen = bFullscreen;
-
-		ViewInfo->PixelFormat = bHDREnabled ? GRHIHDRDisplayOutputFormat : ViewInfo->SDRPixelFormat;
 		ViewInfo->HDRDisplayColorGamut = HDRColorGamut;
 		ViewInfo->HDRDisplayOutputFormat = HDROutputDevice;
-		ViewInfo->bSceneHDREnabled = bHDREnabled;
+		ViewInfo->bDisplayFormatIsHDR = bHDREnabled;
+		ViewInfo->PixelFormat = GetViewportPixelFormat(*Window, bHDREnabled);
 
 		PreResizeBackBufferDelegate.Broadcast(&ViewInfo->ViewportRHI);
 		if (IsValidRef(ViewInfo->ViewportRHI))
 		{
 			ensureMsgf(ViewInfo->ViewportRHI->GetRefCount() == 1, TEXT("Viewport backbuffer was not properly released"));
-			RHIResizeViewport(ViewInfo->ViewportRHI, NewWidth, NewHeight, bFullscreen, ViewInfo->PixelFormat);
+			RHIResizeViewport(ViewInfo->ViewportRHI, ExtentToResizeTo.X, ExtentToResizeTo.Y, bFullscreen, ViewInfo->PixelFormat);
 		}
 		else
 		{
-			ViewInfo->ViewportRHI = RHICreateViewport(ViewInfo->OSWindow, NewWidth, NewHeight, bFullscreen, ViewInfo->PixelFormat);
+			ViewInfo->ViewportRHI = RHICreateViewport(ViewInfo->OSWindow, ExtentToResizeTo.X, ExtentToResizeTo.Y, bFullscreen, ViewInfo->PixelFormat);
 		}
 
 		PostResizeBackBufferDelegate.Broadcast(&ViewInfo->ViewportRHI);
@@ -539,6 +450,28 @@ void FSlateRHIRenderer::ConditionalResizeViewport(FViewportInfo* ViewInfo, uint3
 	}
 }
 
+EPixelFormat FSlateRHIRenderer::GetViewportPixelFormat(const SWindow& Window, bool bDisplayFormatIsHDR)
+{
+	// Use the configured HDR format if enabled.
+	if (bDisplayFormatIsHDR)
+	{
+		return GRHIHDRDisplayOutputFormat;
+	}
+
+	// Use a known default format in VR / Mobile / Transparent Window SDR configurations.
+	if (bIsStandaloneStereoOnlyDevice || GMaxRHIFeatureLevel == ERHIFeatureLevel::ES3_1
+#if ALPHA_BLENDED_WINDOWS
+		|| Window.GetTransparencySupport() == EWindowTransparency::PerPixel
+#endif
+		)
+	{
+		return GetSlateRecommendedColorFormat();
+	}
+
+	// Let the RHI decide.
+	return PF_Unknown;
+}
+
 void FSlateRHIRenderer::OnVirtualDesktopSizeChanged(const FDisplayMetrics& NewDisplayMetric)
 {
 	// Defer the update to as we need to call FlushRenderingCommands() before sending the event to the RHI. 
@@ -549,7 +482,7 @@ void FSlateRHIRenderer::OnVirtualDesktopSizeChanged(const FDisplayMetrics& NewDi
 
 void FSlateRHIRenderer::UpdateFullscreenState(const TSharedRef<SWindow> Window, uint32 OverrideResX, uint32 OverrideResY)
 {
-	FViewportInfo* ViewInfo = WindowToViewportInfo.FindRef(&Window.Get());
+	FSlateViewportInfo* ViewInfo = WindowToViewportInfo.FindRef(&Window.Get());
 
 	if (!ViewInfo)
 	{
@@ -561,18 +494,16 @@ void FSlateRHIRenderer::UpdateFullscreenState(const TSharedRef<SWindow> Window,
 	if (ViewInfo)
 	{
 		const bool bFullscreen = IsViewportFullscreen(*Window);
+		const bool bIsRenderingStereo = GEngine && GEngine->XRSystem.IsValid() && GEngine->StereoRenderingDevice.IsValid() && GEngine->StereoRenderingDevice->IsStereoEnabled();
 
-		uint32 ResX = OverrideResX ? OverrideResX : GSystemResolution.ResX;
-		uint32 ResY = OverrideResY ? OverrideResY : GSystemResolution.ResY;
+		FIntPoint ExtentToResizeTo(OverrideResX ? OverrideResX : GSystemResolution.ResX, OverrideResY ? OverrideResY : GSystemResolution.ResY);
 
-		bool bIsRenderingStereo = GEngine && GEngine->XRSystem.IsValid() && GEngine->StereoRenderingDevice.IsValid() && GEngine->StereoRenderingDevice->IsStereoEnabled();
 		if ((GIsEditor && Window->IsViewportSizeDrivenByWindow()) || (Window->GetWindowMode() == EWindowMode::WindowedFullscreen) || bIsRenderingStereo)
 		{
-			ResX = ViewInfo->DesiredWidth;
-			ResY = ViewInfo->DesiredHeight;
+			ExtentToResizeTo = ViewInfo->ExtentToResizeTo;
 		}
 
-		ConditionalResizeViewport(ViewInfo, ResX, ResY, bFullscreen, &Window.Get());
+		ResizeViewportIfNeeded(ViewInfo, ExtentToResizeTo, bFullscreen, &Window.Get());
 	}
 }
 
@@ -593,1047 +524,484 @@ void FSlateRHIRenderer::RestoreSystemResolution(const TSharedRef<SWindow> InWind
 	}
 }
 
-/** Called when a window is destroyed to give the renderer a chance to free resources */
 void FSlateRHIRenderer::OnWindowDestroyed(const TSharedRef<SWindow>& InWindow)
 {
-	checkSlow(IsThreadSafeForSlateRendering());
-
-	FViewportInfo** ViewportInfoPtr = WindowToViewportInfo.Find(&InWindow.Get());
-	if (ViewportInfoPtr)
+	if (FSlateViewportInfo** ViewportInfoPtr = WindowToViewportInfo.Find(&InWindow.Get()))
 	{
-		OnSlateWindowDestroyedDelegate.Broadcast(&(*ViewportInfoPtr)->ViewportRHI);
-
-		// Need to flush rendering commands as the viewport may be in use by the render thread
-		// and the rendering resources must be released on the render thread before the viewport can be deleted
-		FlushRenderingCommands();
+		FSlateViewportInfo* ViewportInfo = *ViewportInfoPtr;
 
-		BeginReleaseResource(*ViewportInfoPtr);
+		OnSlateWindowDestroyedDelegate.Broadcast(&ViewportInfo->ViewportRHI);
 
-		// Flush rendering commands again so that the resource deletion request is processed.
+		// Perform the release in lock-step with the render thread to simplify resource lifetimes.
+		FlushRenderingCommands();
+		BeginReleaseResource(ViewportInfo);
 		FlushRenderingCommands();
+		delete ViewportInfo;
 
-		delete *ViewportInfoPtr;
+		WindowToViewportInfo.Remove(&InWindow.Get());
 	}
-
-	WindowToViewportInfo.Remove(&InWindow.Get());
 }
 
-/** Called when a window is Finished being Reshaped - Currently need to check if its HDR status has changed */
 void FSlateRHIRenderer::OnWindowFinishReshaped(const TSharedPtr<SWindow>& InWindow)
 {
-	FViewportInfo* ViewInfo = WindowToViewportInfo.FindRef(InWindow.Get());
+	FSlateViewportInfo* ViewInfo = WindowToViewportInfo.FindRef(InWindow.Get());
 	RHICheckViewportHDRStatus(ViewInfo->ViewportRHI);
 }
 
 // Limited platform support for HDR UI composition
-bool SupportsUICompositionRendering(const EShaderPlatform Platform)
+bool SupportsCompositeUIWithSceneHDR(const EShaderPlatform Platform)
 {
 	return IsFeatureLevelSupported(Platform, ERHIFeatureLevel::SM5) && (RHISupportsGeometryShaders(Platform) || RHISupportsVertexShaderLayer(Platform));
 }
 
-// Pixel shader to generate LUT for HDR UI composition
-class FCompositeLUTGenerationPS : public FGlobalShader
+bool CompositeUIWithSceneHDR()
 {
-	DECLARE_SHADER_TYPE(FCompositeLUTGenerationPS, Global);
-public:
-
-	static bool ShouldCompilePermutation(const FGlobalShaderPermutationParameters& Parameters)
-	{
-		return SupportsUICompositionRendering(Parameters.Platform);
-	}
-
-	FCompositeLUTGenerationPS(const ShaderMetaType::CompiledShaderInitializerType& Initializer) :
-		FGlobalShader(Initializer)
-	{
-		OutputDevice.Bind(Initializer.ParameterMap, TEXT("OutputDevice"));
-		OutputGamut.Bind(Initializer.ParameterMap, TEXT("OutputGamut"));
-		OutputMaxLuminance.Bind(Initializer.ParameterMap, TEXT("OutputMaxLuminance"));
-		ACESMinMaxData.Bind(Initializer.ParameterMap, TEXT("ACESMinMaxData"));
-		ACESMidData.Bind(Initializer.ParameterMap, TEXT("ACESMidData"));
-		ACESCoefsLow_0.Bind(Initializer.ParameterMap, TEXT("ACESCoefsLow_0"));
-		ACESCoefsHigh_0.Bind(Initializer.ParameterMap, TEXT("ACESCoefsHigh_0"));
-		ACESCoefsLow_4.Bind(Initializer.ParameterMap, TEXT("ACESCoefsLow_4"));
-		ACESCoefsHigh_4.Bind(Initializer.ParameterMap, TEXT("ACESCoefsHigh_4"));
-		ACESSceneColorMultiplier.Bind(Initializer.ParameterMap, TEXT("ACESSceneColorMultiplier"));
-	}
-	FCompositeLUTGenerationPS() {}
-
-	void SetParameters(FRHIBatchedShaderParameters& BatchedParameters, EDisplayOutputFormat DisplayOutputFormat, EDisplayColorGamut DisplayColorGamut, float DisplayMaxLuminance)
-	{
-		static const auto CVarOutputGamma = IConsoleManager::Get().FindTConsoleVariableDataFloat(TEXT("r.TonemapperGamma"));
-
-		int32 OutputDeviceValue = (int32)DisplayOutputFormat;
-		int32 OutputGamutValue = (int32)DisplayColorGamut;
-		float Gamma = CVarOutputGamma->GetValueOnRenderThread();
-
-        // In case gamma is unspecified, fall back to 2.2 which is the most common case
-		if ((PLATFORM_APPLE || DisplayOutputFormat == EDisplayOutputFormat::SDR_ExplicitGammaMapping) && Gamma == 0.0f)
-		{
-			Gamma = 2.2f;
-		}
-
-		if (Gamma > 0.0f)
-		{
-			// Enforce user-controlled ramp over sRGB or Rec709
-			OutputDeviceValue = FMath::Max(OutputDeviceValue, (int32)EDisplayOutputFormat::SDR_ExplicitGammaMapping);
-		}
-
-		SetShaderValue(BatchedParameters, OutputDevice, OutputDeviceValue);
-		SetShaderValue(BatchedParameters, OutputGamut, OutputGamutValue);
-		SetShaderValue(BatchedParameters, OutputMaxLuminance, DisplayMaxLuminance);
-
-		FACESTonemapParams TmpInternalACESParams;
-		GetACESTonemapParameters(TmpInternalACESParams);
-
-		SetShaderValue(BatchedParameters, ACESMinMaxData, TmpInternalACESParams.ACESMinMaxData);
-		SetShaderValue(BatchedParameters, ACESMidData, TmpInternalACESParams.ACESMidData);
-		SetShaderValue(BatchedParameters, ACESCoefsLow_0, TmpInternalACESParams.ACESCoefsLow_0);
-		SetShaderValue(BatchedParameters, ACESCoefsHigh_0, TmpInternalACESParams.ACESCoefsHigh_0);
-		SetShaderValue(BatchedParameters, ACESCoefsLow_4, TmpInternalACESParams.ACESCoefsLow_4);
-		SetShaderValue(BatchedParameters, ACESCoefsHigh_4, TmpInternalACESParams.ACESCoefsHigh_4);
-		SetShaderValue(BatchedParameters, ACESSceneColorMultiplier, TmpInternalACESParams.ACESSceneColorMultiplier);
-
-	}
-
-	static const TCHAR* GetSourceFilename()
-	{
-		return TEXT("/Engine/Private/CompositeUIPixelShader.usf");
-	}
-
-	static const TCHAR* GetFunctionName()
-	{
-		return TEXT("Main");
-	}
+	// Optional off-screen UI composition during HDR rendering
+	static const auto CVarCompositeMode = IConsoleManager::Get().FindTConsoleVariableDataInt(TEXT("r.HDR.UI.CompositeMode"));
 
-private:
-	LAYOUT_FIELD(FShaderParameter, OutputDevice);
-	LAYOUT_FIELD(FShaderParameter, OutputGamut);
-	LAYOUT_FIELD(FShaderParameter, OutputMaxLuminance);
-	LAYOUT_FIELD(FShaderParameter, ACESMinMaxData) // xy = min ACES/luminance, zw = max ACES/luminance
-	LAYOUT_FIELD(FShaderParameter, ACESMidData) // x = mid ACES, y = mid luminance, z = mid slope
-	LAYOUT_FIELD(FShaderParameter, ACESCoefsLow_0) // coeflow 0-3
-	LAYOUT_FIELD(FShaderParameter, ACESCoefsHigh_0) // coefhigh 0-3
-	LAYOUT_FIELD(FShaderParameter, ACESCoefsLow_4)
-	LAYOUT_FIELD(FShaderParameter, ACESCoefsHigh_4)
-	LAYOUT_FIELD(FShaderParameter, ACESSceneColorMultiplier)
-};
+	return GRHISupportsHDROutput
+		&& GSupportsVolumeTextureRendering
+		&& SupportsCompositeUIWithSceneHDR(GetFeatureLevelShaderPlatform(GMaxRHIFeatureLevel))
+		&& CVarCompositeMode
+		&& CVarCompositeMode->GetValueOnAnyThread() != 0;
+}
 
-IMPLEMENT_SHADER_TYPE(, FCompositeLUTGenerationPS, TEXT("/Engine/Private/CompositeUIPixelShader.usf"), TEXT("GenerateLUTPS"), SF_Pixel);
+BEGIN_SHADER_PARAMETER_STRUCT(FCompositeShaderCommonParameters, )
+	SHADER_PARAMETER_RDG_TEXTURE(Texture2D, UITexture)
+	SHADER_PARAMETER_SAMPLER(SamplerState, UISampler)
+	SHADER_PARAMETER_RDG_TEXTURE(Texture2D<uint>, UIWriteMaskTexture)
+	SHADER_PARAMETER(float, UILevel)
+	SHADER_PARAMETER(float, UILuminance)
+	SHADER_PARAMETER(float, ColorVisionDeficiencyType)
+	SHADER_PARAMETER(float, ColorVisionDeficiencySeverity)
+	SHADER_PARAMETER(float, bCorrectDeficiency)
+	SHADER_PARAMETER(float, bSimulateCorrectionWithDeficiency)
+END_SHADER_PARAMETER_STRUCT()
 
-// Pixel shader to composite UI over HDR buffer
-class FCompositeShaderBase : public FGlobalShader
+class FCompositeShader : public FGlobalShader
 {
-	DECLARE_TYPE_LAYOUT(FCompositeShaderBase, NonVirtual);
 public:
+	class FSCRGBEncoding : SHADER_PERMUTATION_BOOL("SCRGB_ENCODING");
+	class FApplyColorDeficiency : SHADER_PERMUTATION_BOOL("APPLY_COLOR_DEFICIENCY");
+	using FPermutationDomain = TShaderPermutationDomain<FSCRGBEncoding, FApplyColorDeficiency>;
 
 	static bool ShouldCompilePermutation(const FGlobalShaderPermutationParameters& Parameters)
 	{
-		return SupportsUICompositionRendering(Parameters.Platform);
-	}
-
-	FCompositeShaderBase(const ShaderMetaType::CompiledShaderInitializerType& Initializer) :
-		FGlobalShader(Initializer)
-	{
-		UITexture.Bind(Initializer.ParameterMap, TEXT("UITexture"));
-		UIWriteMaskTexture.Bind(Initializer.ParameterMap, TEXT("UIWriteMaskTexture"));
-		UISampler.Bind(Initializer.ParameterMap, TEXT("UISampler"));
-		ColorSpaceLUT.Bind(Initializer.ParameterMap, TEXT("ColorSpaceLUT"));
-		ColorSpaceLUTSampler.Bind(Initializer.ParameterMap, TEXT("ColorSpaceLUTSampler"));
-		UILevel.Bind(Initializer.ParameterMap, TEXT("UILevel"));
-		UILuminance.Bind(Initializer.ParameterMap, TEXT("UILuminance"));
-		OutputDevice.Bind(Initializer.ParameterMap, TEXT("OutputDevice"));
-		ColorVisionDeficiencyType.Bind(Initializer.ParameterMap, TEXT("ColorVisionDeficiencyType"));
-		ColorVisionDeficiencySeverity.Bind(Initializer.ParameterMap, TEXT("ColorVisionDeficiencySeverity"));
-		bCorrectDeficiency.Bind(Initializer.ParameterMap, TEXT("bCorrectDeficiency"));
-		bSimulateCorrectionWithDeficiency.Bind(Initializer.ParameterMap, TEXT("bSimulateCorrectionWithDeficiency"));
-	}
-	FCompositeShaderBase() = default;
-
-	void SetParametersBase(FRHIBatchedShaderParameters& BatchedParameters, FRHITexture* UITextureRHI, FRHITexture* UITextureWriteMaskRHI, FRHITexture* ColorSpaceLUTRHI)
-	{
-		static const auto CVarOutputDevice = IConsoleManager::Get().FindTConsoleVariableDataInt(TEXT("r.HDR.Display.OutputDevice"));
-
-		SetTextureParameter(BatchedParameters, UITexture, UISampler, TStaticSamplerState<SF_Point>::GetRHI(), UITextureRHI);
-		SetTextureParameter(BatchedParameters, ColorSpaceLUT, ColorSpaceLUTSampler, TStaticSamplerState<SF_Bilinear, AM_Clamp, AM_Clamp, AM_Clamp>::GetRHI(), ColorSpaceLUTRHI);
-		SetShaderValue(BatchedParameters, UILevel, CVarUILevel.GetValueOnRenderThread());
-		SetShaderValue(BatchedParameters, OutputDevice, CVarOutputDevice->GetValueOnRenderThread());
-		SetShaderValue(BatchedParameters, UILuminance, CVarHDRUILuminance.GetValueOnRenderThread());
-
-		if (UITextureWriteMaskRHI != nullptr)
-		{
-			SetTextureParameter(BatchedParameters, UIWriteMaskTexture, UITextureWriteMaskRHI);
-		}
-	}
-
-	void SetColorDeficiencyParamsBase(FRHIBatchedShaderParameters& BatchedParameters, bool bCorrect, EColorVisionDeficiency DeficiencyType, int32 Severity, bool bShowCorrectionWithDeficiency)
-	{
-		SetShaderValue(BatchedParameters, ColorVisionDeficiencyType, (float)DeficiencyType);
-		SetShaderValue(BatchedParameters, ColorVisionDeficiencySeverity, (float)Severity);
-		SetShaderValue(BatchedParameters, bCorrectDeficiency, bCorrect ? 1.0f : 0.0f);
-		SetShaderValue(BatchedParameters, bSimulateCorrectionWithDeficiency, bShowCorrectionWithDeficiency ? 1.0f : 0.0f);
-	}
-
-	static const TCHAR* GetSourceFilename()
-	{
-		return TEXT("/Engine/Private/CompositeUIPixelShader.usf");
-	}
-
-private:
-	LAYOUT_FIELD(FShaderResourceParameter, UITexture);
-	LAYOUT_FIELD(FShaderResourceParameter, UIWriteMaskTexture);
-	LAYOUT_FIELD(FShaderResourceParameter, UISampler);
-	LAYOUT_FIELD(FShaderResourceParameter, ColorSpaceLUT);
-	LAYOUT_FIELD(FShaderResourceParameter, ColorSpaceLUTSampler);
-	LAYOUT_FIELD(FShaderParameter, UILevel);
-	LAYOUT_FIELD(FShaderParameter, UILuminance);
-	LAYOUT_FIELD(FShaderParameter, OutputDevice);
-	LAYOUT_FIELD(FShaderParameter, ColorVisionDeficiencyType);
-	LAYOUT_FIELD(FShaderParameter, ColorVisionDeficiencySeverity);
-	LAYOUT_FIELD(FShaderParameter, bCorrectDeficiency);
-	LAYOUT_FIELD(FShaderParameter, bSimulateCorrectionWithDeficiency);
-};
-IMPLEMENT_TYPE_LAYOUT(FCompositeShaderBase);
-
-class FCompositePSBase : public FCompositeShaderBase
-{
-	DECLARE_TYPE_LAYOUT(FCompositePSBase, NonVirtual);
-public:
-	FCompositePSBase(const ShaderMetaType::CompiledShaderInitializerType& Initializer) : FCompositeShaderBase(Initializer)
-	{
-		SceneTexture.Bind(Initializer.ParameterMap, TEXT("SceneTexture"));
-		SceneSampler.Bind(Initializer.ParameterMap, TEXT("SceneSampler"));
-	}
-	FCompositePSBase() = default;
-
-	static const TCHAR* GetFunctionName()
-	{
-		return TEXT("Main");
-	}
-
-	void SetParameters(
-		FRHIBatchedShaderParameters& BatchedParameters,
-		FRHITexture* UITextureRHI,
-		FRHITexture* UITextureWriteMaskRHI,
-		FRHITexture* SceneTextureRHI,
-		FRHITexture* ColorSpaceLUTRHI,
-		bool bCorrect,
-		EColorVisionDeficiency DeficiencyType,
-		int32 Severity,
-		bool bShowCorrectionWithDeficiency
-	)
-	{
-		SetParametersBase(BatchedParameters, UITextureRHI, UITextureWriteMaskRHI, ColorSpaceLUTRHI);
-		SetTextureParameter(BatchedParameters, SceneTexture, SceneSampler, TStaticSamplerState<SF_Point>::GetRHI(), SceneTextureRHI);
-		SetColorDeficiencyParamsBase(BatchedParameters, bCorrect, DeficiencyType, Severity, bShowCorrectionWithDeficiency);
+		return SupportsCompositeUIWithSceneHDR(Parameters.Platform);
 	}
-
-private:
-	LAYOUT_FIELD(FShaderResourceParameter, SceneTexture);
-	LAYOUT_FIELD(FShaderResourceParameter, SceneSampler);
 };
-IMPLEMENT_TYPE_LAYOUT(FCompositePSBase);
 
-template<uint32 EncodingType, bool bApplyColorDeficiency>
-class FCompositePS : public FCompositePSBase
+class FCompositePS : public FCompositeShader
 {
-	DECLARE_SHADER_TYPE(FCompositePS, Global);
 public:
-	FCompositePS(const ShaderMetaType::CompiledShaderInitializerType& Initializer) : FCompositePSBase(Initializer) {}
-	FCompositePS() = default;
-
-	static void ModifyCompilationEnvironment(const FGlobalShaderPermutationParameters& Parameters, FShaderCompilerEnvironment& OutEnvironment)
-	{
-		FGlobalShader::ModifyCompilationEnvironment(Parameters, OutEnvironment);
-		OutEnvironment.SetDefine(TEXT("SCRGB_ENCODING"), EncodingType);
-		OutEnvironment.SetDefine(TEXT("APPLY_COLOR_DEFICIENCY"), uint32(bApplyColorDeficiency ? 1 : 0));
-	}
+	DECLARE_GLOBAL_SHADER(FCompositePS);
+	SHADER_USE_PARAMETER_STRUCT(FCompositePS, FCompositeShader);
+
+	BEGIN_SHADER_PARAMETER_STRUCT(FParameters, )
+		SHADER_PARAMETER_STRUCT_INCLUDE(FCompositeShaderCommonParameters, Common)
+		SHADER_PARAMETER_RDG_TEXTURE(Texture2D, SceneTexture)
+		SHADER_PARAMETER_SAMPLER(SamplerState, SceneSampler)
+		RENDER_TARGET_BINDING_SLOTS()
+	END_SHADER_PARAMETER_STRUCT()
 };
 
-typedef FCompositePS<0, false> FCompositePS_PQ;
-typedef FCompositePS<0, true> FCompositePS_PQ_CDV;
-typedef FCompositePS<1, false> FCompositePS_SCRGB;
-typedef FCompositePS<1, true> FCompositePS_SCRGB_CDV;
-
-IMPLEMENT_SHADER_TYPE(template<>, FCompositePS_PQ, FCompositePSBase::GetSourceFilename(), FCompositePSBase::GetFunctionName(), SF_Pixel);
-IMPLEMENT_SHADER_TYPE(template<>, FCompositePS_PQ_CDV,  FCompositePSBase::GetSourceFilename(), FCompositePSBase::GetFunctionName(), SF_Pixel);
-IMPLEMENT_SHADER_TYPE(template<>, FCompositePS_SCRGB, FCompositePSBase::GetSourceFilename(), FCompositePSBase::GetFunctionName(), SF_Pixel);
-IMPLEMENT_SHADER_TYPE(template<>, FCompositePS_SCRGB_CDV,  FCompositePSBase::GetSourceFilename(), FCompositePSBase::GetFunctionName(), SF_Pixel);
+IMPLEMENT_GLOBAL_SHADER(FCompositePS, "/Engine/Private/CompositeUIPixelShader.usf", "Main", SF_Pixel);
 
-class FCompositeCSBase : public FCompositeShaderBase
+class FCompositeCS : public FCompositeShader
 {
-	DECLARE_TYPE_LAYOUT(FCompositeCSBase, NonVirtual);
 public:
-
 	static const uint32 NUM_THREADS_PER_GROUP = 16;
 
+	DECLARE_GLOBAL_SHADER(FCompositeCS);
+	SHADER_USE_PARAMETER_STRUCT(FCompositeCS, FCompositeShader);
+	
+	BEGIN_SHADER_PARAMETER_STRUCT(FParameters, )
+		SHADER_PARAMETER_STRUCT_INCLUDE(FCompositeShaderCommonParameters, Common)
+		SHADER_PARAMETER_RDG_TEXTURE_UAV(RWTexture2D<float4>, RWSceneTexture)
+		SHADER_PARAMETER(FVector4f, SceneTextureDimensions)
+	END_SHADER_PARAMETER_STRUCT()
+
 	static bool IsShaderSupported(const EShaderPlatform ShaderPlatform)
 	{
 		return RHISupports4ComponentUAVReadWrite(ShaderPlatform) && RHISupportsSwapchainUAVs(ShaderPlatform);
 	}
 
-	FCompositeCSBase(const ShaderMetaType::CompiledShaderInitializerType& Initializer) : FCompositeShaderBase(Initializer)
-	{
-		RWSceneTexture.Bind(Initializer.ParameterMap, TEXT("RWSceneTexture"));
-		SceneTextureDimensions.Bind(Initializer.ParameterMap, TEXT("SceneTextureDimensions"));
-	}
-	FCompositeCSBase() = default;
-
-	static const TCHAR* GetFunctionName()
-	{
-		return TEXT("CompositeUICS");
-	}
-
-	void SetParameters(
-		FRHIBatchedShaderParameters& BatchedParameters,
-		FRHITexture* UITextureRHI,
-		FRHITexture* UITextureWriteMaskRHI,
-		FRHIUnorderedAccessView* InRWSceneTexture,
-		FRHITexture* ColorSpaceLUTRHI,
-		const FVector4f& InSceneTextureDimensions,
-		bool bCorrect,
-		EColorVisionDeficiency DeficiencyType,
-		int32 Severity,
-		bool bShowCorrectionWithDeficiency
-	)
-	{
-		SetParametersBase(BatchedParameters, UITextureRHI, UITextureWriteMaskRHI, ColorSpaceLUTRHI);
-		SetUAVParameter(BatchedParameters, RWSceneTexture, InRWSceneTexture);
-		SetShaderValue(BatchedParameters, SceneTextureDimensions, InSceneTextureDimensions);
-		SetColorDeficiencyParamsBase(BatchedParameters, bCorrect, DeficiencyType, Severity, bShowCorrectionWithDeficiency);
-	}
-
-private:
-	LAYOUT_FIELD(FShaderResourceParameter, RWSceneTexture);
-	LAYOUT_FIELD(FShaderParameter, SceneTextureDimensions);
-};
-IMPLEMENT_TYPE_LAYOUT(FCompositeCSBase);
-
-template<uint32 EncodingType, bool bApplyColorDeficiency>
-class FCompositeCS : public FCompositeCSBase
-{
-	DECLARE_SHADER_TYPE(FCompositeCS, Global);
-public:
-	FCompositeCS(const ShaderMetaType::CompiledShaderInitializerType& Initializer) : FCompositeCSBase(Initializer) {}
-	FCompositeCS() = default;
-
 	static bool ShouldCompilePermutation(const FGlobalShaderPermutationParameters& Parameters)
 	{
-		if (!SupportsUICompositionRendering(Parameters.Platform))
-		{
-			return false;
-		}
-
-		return IsShaderSupported(Parameters.Platform);
+		return FCompositeShader::ShouldCompilePermutation(Parameters) && IsShaderSupported(Parameters.Platform);
 	}
 
-
 	static void ModifyCompilationEnvironment(const FGlobalShaderPermutationParameters& Parameters, FShaderCompilerEnvironment& OutEnvironment)
 	{
 		FGlobalShader::ModifyCompilationEnvironment(Parameters, OutEnvironment);
 		OutEnvironment.SetDefine(TEXT("USE_COMPUTE_FOR_COMPOSITION"), 1);
-		OutEnvironment.SetDefine(TEXT("SCRGB_ENCODING"), EncodingType);
 		OutEnvironment.SetDefine(TEXT("NUM_THREADS_PER_GROUP"), NUM_THREADS_PER_GROUP);
-		OutEnvironment.SetDefine(TEXT("APPLY_COLOR_DEFICIENCY"), uint32(bApplyColorDeficiency ? 1 : 0));
 	}
 };
 
-typedef FCompositeCS<0, false> FCompositeCS_PQ;
-typedef FCompositeCS<0, true>  FCompositeCS_PQ_CDV;
-typedef FCompositeCS<1, false> FCompositeCS_SCRGB;
-typedef FCompositeCS<1, true>  FCompositeCS_SCRGB_CDV;
-
-IMPLEMENT_SHADER_TYPE(template<>, FCompositeCS_PQ,        FCompositeCSBase::GetSourceFilename(), FCompositeCSBase::GetFunctionName(), SF_Compute);
-IMPLEMENT_SHADER_TYPE(template<>, FCompositeCS_PQ_CDV,    FCompositeCSBase::GetSourceFilename(), FCompositeCSBase::GetFunctionName(), SF_Compute);
-IMPLEMENT_SHADER_TYPE(template<>, FCompositeCS_SCRGB,     FCompositeCSBase::GetSourceFilename(), FCompositeCSBase::GetFunctionName(), SF_Compute);
-IMPLEMENT_SHADER_TYPE(template<>, FCompositeCS_SCRGB_CDV, FCompositeCSBase::GetSourceFilename(), FCompositeCSBase::GetFunctionName(), SF_Compute);
+IMPLEMENT_GLOBAL_SHADER(FCompositeCS, "/Engine/Private/CompositeUIPixelShader.usf", "CompositeUICS", SF_Compute);
 
-
-int32 SlateWireFrame = 0;
-static FAutoConsoleVariableRef CVarSlateWireframe(TEXT("Slate.ShowWireFrame"), SlateWireFrame, TEXT(""), ECVF_Default);
-
-void RenderSlateBatch(FTextureRHIRef SlateRenderTarget, bool bClear, bool bIsHDR, FViewportInfo& ViewportInfo, const FMatrix& ViewMatrix, FSlateBatchData& BatchData, FRHICommandListImmediate& RHICmdList,
-				      const uint32 ViewportWidth, const uint32 ViewportHeight, const struct FSlateDrawWindowCommandParams& DrawCommandParams, TSharedPtr<FSlateRHIRenderingPolicy> RenderingPolicy, 
-					  FTextureRHIRef PostProcessBuffer)
+void FSlateRHIRenderer::DrawWindow_RenderThread(FRDGBuilder& GraphBuilder, const FSlateDrawWindowPassInputs& Inputs)
 {
-	FRHIRenderPassInfo RPInfo(SlateRenderTarget, ERenderTargetActions::Load_Store);
-
-	if (bClear)
-	{
-		RPInfo.ColorRenderTargets[0].Action = ERenderTargetActions::Clear_Store;
-	}
-
-	if (ViewportInfo.bRequiresStencilTest)
-	{
-		check(IsValidRef(ViewportInfo.DepthStencil));
-
-		ERenderTargetActions StencilAction = ERenderTargetActions::DontLoad_Store;
-		RPInfo.DepthStencilRenderTarget.Action = MakeDepthStencilTargetActions(ERenderTargetActions::DontLoad_DontStore, StencilAction);
-		RPInfo.DepthStencilRenderTarget.DepthStencilTarget = ViewportInfo.DepthStencil;
-		RPInfo.DepthStencilRenderTarget.ExclusiveDepthStencil = FExclusiveDepthStencil::DepthNop_StencilWrite;
-	}
-
-#if WITH_SLATE_VISUALIZERS
-	if (CVarShowSlateBatching.GetValueOnRenderThread() != 0 || CVarShowSlateOverdraw.GetValueOnRenderThread() != 0)
-	{
-		RPInfo.ColorRenderTargets[0].Action = ERenderTargetActions::Clear_Store;
-		if (ViewportInfo.bRequiresStencilTest)
-		{
-			// Reset the backbuffer as our color render target and also set a depth stencil buffer
-			ERenderTargetActions StencilAction = ERenderTargetActions::Clear_Store;
-			RPInfo.DepthStencilRenderTarget.Action = MakeDepthStencilTargetActions(ERenderTargetActions::Load_Store, StencilAction);
-			RPInfo.DepthStencilRenderTarget.DepthStencilTarget = ViewportInfo.DepthStencil;
-			RPInfo.DepthStencilRenderTarget.ExclusiveDepthStencil = FExclusiveDepthStencil::DepthWrite_StencilWrite;
-		}
-	}
-#endif
-	{
-		bool bHasBatches = BatchData.GetRenderBatches().Num() > 0;
-		if (bHasBatches || bClear)
-		{
-			TransitionRenderPassTargets(RHICmdList, RPInfo);
-			RHICmdList.BeginRenderPass(RPInfo, TEXT("SlateBatches"));
-			SCOPE_CYCLE_COUNTER(STAT_SlateRTDrawBatches);
-
-			if (bHasBatches)
-			{
-				FSlateBackBuffer SlateBackBuffer(SlateRenderTarget, FIntPoint(ViewportWidth, ViewportHeight));
-
-				FSlateRenderingParams RenderParams(ViewMatrix * ViewportInfo.ProjectionMatrix, DrawCommandParams.Time);
-				RenderParams.bWireFrame = !!SlateWireFrame;
-				RenderParams.bIsHDR = bIsHDR;
-				RenderParams.HDRDisplayColorGamut = ViewportInfo.HDRDisplayColorGamut;
-				RenderParams.ViewRect = DrawCommandParams.ViewRect;
-				RenderParams.CursorPostion = DrawCommandParams.CursorPostion;
-				RenderParams.ViewportScaleUI = DrawCommandParams.ViewportScaleUI;
-				RenderParams.UsedSlatePostBuffers = DrawCommandParams.UsedSlatePostBuffers;
-				if (ViewportInfo.bSceneHDREnabled && !bIsHDR)
-				{
-					RenderParams.UITarget = ViewportInfo.UITargetRT;
-				}
-				RenderingPolicy->SetUseGammaCorrection(!bIsHDR);
-
-				FTextureRHIRef EmptyTarget;
-
-				RenderingPolicy->DrawElements
-				(
-					RHICmdList,
-					SlateBackBuffer,
-					SlateRenderTarget,
-					PostProcessBuffer,
-					ViewportInfo.bRequiresStencilTest ? ViewportInfo.DepthStencil : EmptyTarget,
-					BatchData.GetFirstRenderBatchIndex(),
-					BatchData.GetRenderBatches(),
-					RenderParams
-				);
-			}
-		}
-	}
-
-	// @todo Could really use a refactor.
-	// Kind of gross but we don't want to restart renderpasses for no reason.
-	// If the color deficiency shaders are active within DrawElements there will not be a renderpass here.
-	// In the general case there will be a RenderPass active at this point.
-	if (RHICmdList.IsInsideRenderPass())
-	{
-		RHICmdList.EndRenderPass();
-	}
-}
-
-inline bool CompositeUIWithHdrRenderTarget(const FViewportInfo* ViewInfo)
-{
-	// Optional off-screen UI composition during HDR rendering
-	static const auto CVarCompositeMode = IConsoleManager::Get().FindTConsoleVariableDataInt(TEXT("r.HDR.UI.CompositeMode"));
+	LLM_SCOPE(ELLMTag::SceneRender);
 
-	const bool bSupportsUIComposition = GRHISupportsHDROutput && GSupportsVolumeTextureRendering && SupportsUICompositionRendering(GetFeatureLevelShaderPlatform(GMaxRHIFeatureLevel));
-	const bool bCompositeUI = bSupportsUIComposition
-		&& CVarCompositeMode && CVarCompositeMode->GetValueOnAnyThread() != 0
-		&& ViewInfo->bSceneHDREnabled;
+	FSlateViewportInfo& ViewportInfo = *Inputs.ViewportInfo;
+	FSlateWindowElementList& WindowElementList = *Inputs.WindowElementList;
 
-	return bCompositeUI;
-}
+	FMaterialRenderProxy::UpdateDeferredCachedUniformExpressions(GraphBuilder.RHICmdList);
 
-/** Draws windows from a FSlateDrawBuffer on the render thread */
-void FSlateRHIRenderer::DrawWindow_RenderThread(FRHICommandListImmediate& RHICmdList, FViewportInfo& ViewportInfo, FSlateWindowElementList& WindowElementList, const struct FSlateDrawWindowCommandParams& DrawCommandParams)
-{
-	LLM_SCOPE(ELLMTag::SceneRender);
-	
-	bool bRenderOffscreen = false;	// Render to an offscreen texture which can then be finally color converted at the end.
-	
-#if WITH_EDITOR
-	if (RHIGetColorSpace(ViewportInfo.ViewportRHI) != EColorSpaceAndEOTF::ERec709_sRGB)
+	AddPass(GraphBuilder, RDG_EVENT_NAME("BeginDrawingViewport"), [ViewportRHI = ViewportInfo.ViewportRHI] (FRHICommandListImmediate& RHICmdList)
 	{
-		bRenderOffscreen = true;
-	}
-#endif
+		GetRendererModule().InitializeSystemTextures(RHICmdList);
+		RHICmdList.BeginDrawingViewport(ViewportRHI, FTextureRHIRef());
+	});
 
-	static uint32 LastTimestamp = FPlatformTime::Cycles();
 	{
-		const FRHIGPUMask PresentingGPUMask = FRHIGPUMask::FromIndex(RHIGetViewportNextPresentGPUIndex(ViewportInfo.ViewportRHI));
-		SCOPED_GPU_MASK(RHICmdList, PresentingGPUMask);
-
-#if WANTS_DRAW_MESH_EVENTS
-		SCOPED_CONDITIONAL_DRAW_EVENTF(RHICmdList, SlateUI,  DrawCommandParams.WindowTitle.IsEmpty(), TEXT("SlateUI Title = <none>"));
-		SCOPED_CONDITIONAL_DRAW_EVENTF(RHICmdList, SlateUI, !DrawCommandParams.WindowTitle.IsEmpty(), TEXT("SlateUI Title = %s"), DrawCommandParams.WindowTitle);
-#else
-		SCOPED_DRAW_EVENT(RHICmdList, SlateUI);
-#endif
-
-		SCOPED_GPU_STAT(RHICmdList, SlateUI);
-		SCOPED_NAMED_EVENT_TEXT("Slate::DrawWindow_RenderThread", FColor::Magenta);
-
-		// Should only be called by the rendering thread
-		check(IsInRenderingThread());
-
-		FMaterialRenderProxy::UpdateDeferredCachedUniformExpressions();
-		GetRendererModule().InitializeSystemTextures(RHICmdList);
+		RDG_GPU_MASK_SCOPE(GraphBuilder, FRHIGPUMask::FromIndex(RHIGetViewportNextPresentGPUIndex(ViewportInfo.ViewportRHI)));
+		RDG_EVENT_SCOPE_CONDITIONAL(GraphBuilder, Inputs.WindowTitle.IsEmpty(), "SlateUI Title = <none>");
+		RDG_EVENT_SCOPE_CONDITIONAL(GraphBuilder, !Inputs.WindowTitle.IsEmpty(),"SlateUI Title = %s", *Inputs.WindowTitle);
+		RDG_GPU_STAT_SCOPE(GraphBuilder, SlateUI);
+		RDG_CSV_STAT_EXCLUSIVE_SCOPE(GraphBuilder, Slate);
+		TRACE_CPUPROFILER_EVENT_SCOPE(Slate::DrawWindow_RenderThread);
 
-		const bool bCompositeUI = CompositeUIWithHdrRenderTarget(&ViewportInfo);
+		FGlobalShaderMap* ShaderMap = GetGlobalShaderMap(GMaxRHIFeatureLevel);
 
-		const int32 CompositionLUTSize = 32;
+		// The viewport texture is an optional user-allocated render target. This is rendered to if valid.
+		FRHITexture* ViewportTextureRHI = ViewportInfo.GetRenderTargetTexture();
 
-		// Only need to update LUT on settings change
-		const int32 HDROutputDevice = (int32)ViewportInfo.HDRDisplayOutputFormat;
-		const int32 HDROutputGamut = (int32)ViewportInfo.HDRDisplayColorGamut;
+		// The swap chain is the final output. This is rendered to if no viewport render target is provided.
+		FRHITexture* SwapChainTextureRHI = RHIGetViewportBackBuffer(ViewportInfo.ViewportRHI);
 
-		bool bLUTStale = ViewportInfo.ColorSpaceLUTOutputDevice != HDROutputDevice || ViewportInfo.ColorSpaceLUTOutputGamut != HDROutputGamut;
+		// Only render to the intermediate viewport render target if stereo rendering is enabled, which we'll then composite later.
+		const bool bCompositeStereoToSwapChain = ViewportTextureRHI && GEngine && GEngine->StereoRenderingDevice.IsValid();
 
-		ViewportInfo.ColorSpaceLUTOutputDevice = HDROutputDevice;
-		ViewportInfo.ColorSpaceLUTOutputGamut = HDROutputGamut;
+		// The output texture is what we ultimately render or composite slate elements into.
+		FRHITexture* OutputTextureRHI = bCompositeStereoToSwapChain ? ViewportTextureRHI : SwapChainTextureRHI;
+		FRDGTexture* OutputTexture = RegisterExternalTexture(GraphBuilder, OutputTextureRHI, TEXT("SlateOutputTexture"));
 
-		bool bRenderedStereo = false;
-		if (CVarDrawToVRRenderTarget->GetInt() == 0 && GEngine && IsValidRef(ViewportInfo.GetRenderTargetTexture()) && GEngine->StereoRenderingDevice.IsValid())
-		{
-			const FVector2D WindowSize = WindowElementList.GetWindowSize();
-			GEngine->StereoRenderingDevice->RenderTexture_RenderThread(RHICmdList, RHIGetViewportBackBuffer(ViewportInfo.ViewportRHI), ViewportInfo.GetRenderTargetTexture(), WindowSize);
-			bRenderedStereo = true;
-		}
+		// The elements texture contains UI elements. It can be the same as the output or allocated separately and composited.
+		FRDGTexture* ElementsTexture = OutputTexture;
+		const FIntPoint OutputExtent = OutputTexture->Desc.Extent;
 
+		for (const FSlatePostProcessUpdateRequest& Request : Inputs.PostProcessUpdateRequests)
 		{
-			SCOPED_GPU_STAT(RHICmdList, SlateUI);
-			SCOPE_CYCLE_COUNTER(STAT_SlateRenderingRTTime);
-			CSV_SCOPED_TIMING_STAT_EXCLUSIVE(Slate);
-
-			// should have been created by the game thread
-			check(IsValidRef(ViewportInfo.ViewportRHI));
+			FRHITexture* PostProcessOutputTextureRHI = Request.RenderTargetTextureResource->GetTexture2DRHI();
+			const FScreenPassTexture PostProcessOutputTexture(RegisterExternalTexture(GraphBuilder, PostProcessOutputTextureRHI, TEXT("PostProcessRT")));
 
-			FTextureRHIRef ViewportRT = bRenderedStereo ? nullptr : ViewportInfo.GetRenderTargetTexture();
-			FTextureRHIRef BackBuffer = (ViewportRT) ? ViewportRT : RHIGetViewportBackBuffer(ViewportInfo.ViewportRHI);
-			FTextureRHIRef PostProcessBuffer = BackBuffer;	// If compositing UI then this will be different to the back buffer
+			FIntRect PostProcessInputViewRect;
+			FRHITexture* PostProcessInputTextureRHI;
 
-			const uint32 ViewportWidth = (ViewportRT) ? ViewportRT->GetSizeX() : ViewportInfo.Width;
-			const uint32 ViewportHeight = (ViewportRT) ? ViewportRT->GetSizeY() : ViewportInfo.Height;
-
-			// Check to see that targets are up-to-date
-			if (bCompositeUI && (!ViewportInfo.UITargetRT || ViewportInfo.UITargetRT->GetRHI()->GetSizeX() != ViewportWidth || ViewportInfo.UITargetRT->GetRHI()->GetSizeY() != ViewportHeight))
+			// The viewport texture will contain the scene if it's valid.
+			if (ViewportTextureRHI)
 			{
-				// Composition buffers
-				{
-					ETextureCreateFlags BaseFlags = RHISupportsRenderTargetWriteMask(GMaxRHIShaderPlatform) ? TexCreate_NoFastClearFinalize | TexCreate_DisableDCC : TexCreate_None;
-					FPooledRenderTargetDesc Desc(FPooledRenderTargetDesc::Create2DDesc(FIntPoint(ViewportWidth, ViewportHeight),
-						GetSlateRecommendedColorFormat(),
-						FClearValueBinding::Transparent,
-						BaseFlags,
-						TexCreate_ShaderResource | TexCreate_RenderTargetable,
-						false,
-						1,
-						true,
-						true));
-
-					GRenderTargetPool.FindFreeElement(RHICmdList, Desc, ViewportInfo.UITargetRT, TEXT("UITargetRT"));
-				}
-
-				// LUT
-				{
-					ViewportInfo.ColorSpaceLUT.SafeRelease();
-
-					const FRHITextureCreateDesc Desc =
-						FRHITextureCreateDesc::Create3D(TEXT("ColorSpaceLUT"))
-						.SetExtent(CompositionLUTSize)
-						.SetDepth(CompositionLUTSize)
-						.SetFormat(PF_A2B10G10R10)
-						.SetFlags(ETextureCreateFlags::RenderTargetable | ETextureCreateFlags::ShaderResource)
-						.SetInitialState(ERHIAccess::SRVMask);
-
-					ViewportInfo.ColorSpaceLUT = RHICreateTexture(Desc);
-				}
-
-				bLUTStale = true;
+				PostProcessInputTextureRHI = ViewportTextureRHI;
+				PostProcessInputViewRect.Max = ViewportTextureRHI->GetSizeXY();
 			}
-
-			FTextureRHIRef FinalBuffer = BackBuffer;
-
-			bool bClear = DrawCommandParams.bClear;
-			if (bCompositeUI)
+			else
 			{
-				bClear = true; // Force a clear of the UI buffer to black
-
-#if WITH_EDITOR
-				// in editor mode, we actually don't render to the true backbuffer, but to BufferedRT. The Scene image is rendered in the backbuffer with Slate
-				// We add it with FSlateElementBatcher::AddViewportElement and render it with RenderSlateBatch
-				if (WindowElementList.GetBatchDataHDR().GetRenderBatches().Num() > 0)
-				{
-					FRHIRenderPassInfo RPInfo(FinalBuffer, ERenderTargetActions::Clear_Store);
-					TransitionRenderPassTargets(RHICmdList, RPInfo);
-					RHICmdList.BeginRenderPass(RPInfo, TEXT("Clear back buffer"));
-					RHICmdList.EndRenderPass();
-				}
-#endif
-				// UI backbuffer is temp target
-				BackBuffer = ViewportInfo.UITargetRT->GetRHI();
+				PostProcessInputTextureRHI = SwapChainTextureRHI;
+				PostProcessInputViewRect = Inputs.SceneViewRect;
 			}
 
-#if WITH_EDITOR
-			TRefCountPtr<IPooledRenderTarget> HDRRenderRT;
+			const FScreenPassTexture PostProcessInputTexture(RegisterExternalTexture(GraphBuilder, PostProcessInputTextureRHI, TEXT("ViewportTexture")), PostProcessInputViewRect);
 
-			if (bRenderOffscreen )
+			if (Request.PostProcessorProxy)
 			{
-				FPooledRenderTargetDesc Desc(FPooledRenderTargetDesc::Create2DDesc(FIntPoint(ViewportWidth, ViewportHeight),
-					PF_FloatRGBA,
-					FClearValueBinding::Transparent,
-					TexCreate_None,
-					TexCreate_ShaderResource | TexCreate_RenderTargetable,
-					false,
-					1,
-					true,
-					true));
-
-				GRenderTargetPool.FindFreeElement(RHICmdList, Desc, HDRRenderRT, TEXT("HDRTargetRT"));
-
-				BackBuffer = HDRRenderRT->GetRHI();
+				Request.PostProcessorProxy->PostProcess_Renderthread(GraphBuilder, PostProcessInputTexture, PostProcessOutputTexture);
 			}
-#endif
-
-			if (SlateWireFrame)
+			else
 			{
-				bClear = true;
+				AddDrawTexturePass(GraphBuilder, FScreenPassViewInfo(), PostProcessInputTexture, PostProcessOutputTexture);
 			}
+		}
 
-			RHICmdList.BeginDrawingViewport(ViewportInfo.ViewportRHI, FTextureRHIRef());
-			RHICmdList.SetViewport(0.f, 0.f, 0.f, (float)ViewportWidth, (float)ViewportHeight, 0.0f);
-
-			bool bApplyColorDeficiencyCorrectionToRestore = RenderingPolicy->GetApplyColorDeficiencyCorrection();
-
-			if (bCompositeUI)
-			{
-				RenderingPolicy->SetApplyColorDeficiencyCorrection(false);
-				FSlateBatchData& BatchDataHDR = WindowElementList.GetBatchDataHDR();
-				bool bHasBatches = BatchDataHDR.GetRenderBatches().Num() > 0;
-				if (bHasBatches)
-				{
-					RenderingPolicy->BuildRenderingBuffers(RHICmdList, BatchDataHDR);
-					FTextureRHIRef UITargetHDRRTRHI(FinalBuffer);
-					RenderSlateBatch(UITargetHDRRTRHI, /*bClear*/ false, /*bIsHDR*/ true, ViewportInfo, ViewMatrix, BatchDataHDR, RHICmdList, ViewportWidth, ViewportHeight, DrawCommandParams, RenderingPolicy, PostProcessBuffer);
-			    }
-
-				PostProcessBuffer = FinalBuffer;
-            }
-
-			FSlateBatchData& BatchData = WindowElementList.GetBatchData();
-
-			// Update the vertex and index buffer	
-		    RenderingPolicy->BuildRenderingBuffers(RHICmdList, BatchData);
-
-			// This must happen after rendering buffers are created
-			ViewportInfo.ConditionallyUpdateDepthBuffer(BatchData.IsStencilClippingRequired(), ViewportInfo.DesiredWidth, ViewportInfo.DesiredHeight);
-
-		    bool bHdrTarget = ViewportInfo.bSceneHDREnabled && !bCompositeUI;
-		    RenderSlateBatch(BackBuffer, bClear, bHdrTarget, ViewportInfo, ViewMatrix, BatchData, RHICmdList, ViewportWidth, ViewportHeight, DrawCommandParams, RenderingPolicy, PostProcessBuffer);
+		const bool bCompositeUIWithSceneHDR = ViewportInfo.bDisplayFormatIsHDR && CompositeUIWithSceneHDR();
 
-			if (bCompositeUI)
-			{
-				RenderingPolicy->SetApplyColorDeficiencyCorrection(bApplyColorDeficiencyCorrectionToRestore);
-
-				SCOPED_DRAW_EVENT(RHICmdList, SlateUI_Composition);
+		bool bClearElementsTexture = Inputs.bClear || GSlateWireframe;
+	
+#if WITH_SLATE_VISUALIZERS
+		bClearElementsTexture |= CVarShowSlateBatching.GetValueOnRenderThread() != 0 || CVarShowSlateOverdraw.GetValueOnRenderThread() != 0;
+#endif
 
-				static const FName RendererModuleName("Renderer");
-				IRendererModule& RendererModule = FModuleManager::GetModuleChecked<IRendererModule>(RendererModuleName);
+		if (bCompositeUIWithSceneHDR)
+		{
+			const ETextureCreateFlags WriteMaskFlags = RHISupportsRenderTargetWriteMask(GMaxRHIShaderPlatform) ? ETextureCreateFlags::NoFastClearFinalize | ETextureCreateFlags::DisableDCC : ETextureCreateFlags::None;
 
-				const auto FeatureLevel = GMaxRHIFeatureLevel;
-				auto ShaderMap = GetGlobalShaderMap(FeatureLevel);
+			ElementsTexture = GraphBuilder.CreateTexture(
+				FRDGTextureDesc::Create2D(
+					OutputExtent,
+					GetSlateRecommendedColorFormat(),
+					FClearValueBinding::Transparent,
+					ETextureCreateFlags::ShaderResource | ETextureCreateFlags::RenderTargetable | WriteMaskFlags),
+				TEXT("CompositeUIWithSceneHDRTexture"));
 
-				// Generate composition LUT
-				if (bLUTStale)
-				{
-					// #todo-renderpasses will this touch every pixel? use NoAction?
-					FRHIRenderPassInfo RPInfo(ViewportInfo.ColorSpaceLUT, ERenderTargetActions::Load_Store);
-					RHICmdList.Transition(FRHITransitionInfo(ViewportInfo.ColorSpaceLUT, ERHIAccess::Unknown, ERHIAccess::RTV));
-					RHICmdList.BeginRenderPass(RPInfo, TEXT("GenerateLUT"));
-					{
-						FGraphicsPipelineStateInitializer GraphicsPSOInit;
-						RHICmdList.ApplyCachedRenderTargets(GraphicsPSOInit);
-						GraphicsPSOInit.BlendState = TStaticBlendState<>::GetRHI();
-						GraphicsPSOInit.RasterizerState = TStaticRasterizerState<>::GetRHI();
-						GraphicsPSOInit.DepthStencilState = TStaticDepthStencilState<false, CF_Always>::GetRHI();
-
-						TShaderMapRef<FWriteToSliceVS> VertexShader(ShaderMap);
-						TOptionalShaderMapRef<FWriteToSliceGS> GeometryShader(ShaderMap);
-						TShaderMapRef<FCompositeLUTGenerationPS> PixelShader(ShaderMap);
-						const FVolumeBounds VolumeBounds(CompositionLUTSize);
-
-						GraphicsPSOInit.BoundShaderState.VertexDeclarationRHI = GScreenVertexDeclaration.VertexDeclarationRHI;
-						GraphicsPSOInit.BoundShaderState.VertexShaderRHI = VertexShader.GetVertexShader();
-						GraphicsPSOInit.BoundShaderState.SetGeometryShader(GeometryShader.GetGeometryShader());
-						GraphicsPSOInit.BoundShaderState.PixelShaderRHI = PixelShader.GetPixelShader();
-						GraphicsPSOInit.PrimitiveType = PT_TriangleStrip;
-						SetGraphicsPipelineState(RHICmdList, GraphicsPSOInit, 0);
-
-						SetShaderParametersLegacyVS(RHICmdList, VertexShader, VolumeBounds, FIntVector(VolumeBounds.MaxX - VolumeBounds.MinX));
-
-#if PLATFORM_SUPPORTS_GEOMETRY_SHADERS
-						if (GeometryShader.IsValid())
-						{
-							SetShaderParametersLegacyGS(RHICmdList, GeometryShader, VolumeBounds.MinZ);
-						}
-#endif
-						const float DisplayMaxLuminance = HDRGetDisplayMaximumLuminance();
+			// Force a clear of the UI elements texture to black
+			bClearElementsTexture = true;
+		}
 
-						SetShaderParametersLegacyPS(RHICmdList, PixelShader, ViewportInfo.HDRDisplayOutputFormat, ViewportInfo.HDRDisplayColorGamut, DisplayMaxLuminance);
+		FSlateBatchData& BatchData = WindowElementList.GetBatchData();
+		FSlateBatchData& BatchDataHDR = WindowElementList.GetBatchDataHDR();
 
-						RasterizeToVolumeTexture(RHICmdList, VolumeBounds);
-					}
-					RHICmdList.EndRenderPass();
-					RHICmdList.Transition(FRHITransitionInfo(ViewportInfo.ColorSpaceLUT, ERHIAccess::RTV, ERHIAccess::SRVMask));
-				}
+		const FSlateElementsBuffers SlateElementsBuffers = BuildSlateElementsBuffers(GraphBuilder, BatchData);
+		const FSlateElementsBuffers SlateElementsBuffersHDR = BuildSlateElementsBuffers(GraphBuilder, BatchDataHDR);
 
-				// Composition pass
+		FRDGTexture* SlateStencilTexture = nullptr;
 
-				RHICmdList.Transition(FRHITransitionInfo(ViewportInfo.UITargetRT->GetRHI(), ERHIAccess::Unknown, ERHIAccess::SRVMask));
+		if (BatchData.IsStencilClippingRequired() || BatchDataHDR.IsStencilClippingRequired())
+		{
+			SlateStencilTexture = GraphBuilder.CreateTexture(FRDGTextureDesc::Create2D(OutputExtent, PF_DepthStencil, FClearValueBinding::DepthZero, GetSlateTransientDepthStencilFlags()), TEXT("SlateDepthStencil"));
+		}
 
-				if (RHISupportsRenderTargetWriteMask(GMaxRHIShaderPlatform))
-				{
-					IPooledRenderTarget* RenderTargets[] = { ViewportInfo.UITargetRT.GetReference() };
-					FRenderTargetWriteMask::Decode(RHICmdList, ShaderMap, RenderTargets, ViewportInfo.UITargetRTMask, TexCreate_None, TEXT("UIRTWriteMask"));
-				}
-				FRHITexture* UITargetRTMaskTexture = ViewportInfo.UITargetRTMask.IsValid() ? ViewportInfo.UITargetRTMask->GetRHI() : nullptr;
+		FSlateDrawElementsPassInputs DrawElementsInputs =
+		{
+			  .SceneViewportTexture         = OutputTexture
+			, .ElementsMatrix               = FMatrix44f(ViewportInfo.ProjectionMatrix)
+			, .SceneViewRect                = Inputs.SceneViewRect
+			, .CursorPosition               = Inputs.CursorPosition
+			, .Time                         = Inputs.Time
+			, .HDRDisplayColorGamut         = ViewportInfo.HDRDisplayColorGamut
+			, .UsedSlatePostBuffers         = Inputs.UsedSlatePostBuffers
+			, .ViewportScaleUI              = Inputs.ViewportScaleUI
+			, .bWireframe                   = GSlateWireframe
+			, .bElementsTextureIsHDRDisplay =  ViewportInfo.bDisplayFormatIsHDR
+		};
+
+		if (bCompositeUIWithSceneHDR)
+		{
+			if (!BatchDataHDR.GetRenderBatches().IsEmpty())
+			{
+				DrawElementsInputs.ElementsTexture = OutputTexture;
+				DrawElementsInputs.ElementsLoadAction = ERenderTargetLoadAction::EClear;
+				DrawElementsInputs.ElementsBuffers = SlateElementsBuffersHDR;
+				DrawElementsInputs.StencilTexture = BatchDataHDR.IsStencilClippingRequired() ? SlateStencilTexture : nullptr;
+				DrawElementsInputs.bAllowColorDeficiencyCorrection = false;
 
-				bool bUseScRGB = (HDROutputDevice == (int32)EDisplayOutputFormat::HDR_ACES_1000nit_ScRGB || HDROutputDevice == (int32)EDisplayOutputFormat::HDR_ACES_2000nit_ScRGB);
-				bool bApplyColorDeficiencyFilter = GSlateColorDeficiencyType != EColorVisionDeficiency::NormalVision && GSlateColorDeficiencySeverity > 0;
+				AddSlateDrawElementsPass(GraphBuilder, *RenderingPolicy, DrawElementsInputs, BatchDataHDR.GetRenderBatches(), BatchDataHDR.GetFirstRenderBatchIndex());
+			}
 
-				FRHIUnorderedAccessView* BackBufferUAV = (FCompositeCSBase::IsShaderSupported(GetFeatureLevelShaderPlatform(GMaxRHIFeatureLevel)) && ViewportRT == nullptr) ? RHIGetViewportBackBufferUAV(ViewportInfo.ViewportRHI) : nullptr;
-				if (BackBufferUAV == nullptr)
-				{
-					TRefCountPtr<IPooledRenderTarget> FinalBufferCopy;
-					{
-						FPooledRenderTargetDesc Desc(FPooledRenderTargetDesc::Create2DDesc(FIntPoint(ViewportWidth, ViewportHeight),
-							FinalBuffer->GetFormat(),
-							FClearValueBinding::Transparent,
-							TexCreate_None,
-							TexCreate_ShaderResource,
-							false,
-							1,
-							true,
-							true));
-						GRenderTargetPool.FindFreeElement(RHICmdList, Desc, FinalBufferCopy, TEXT("FinalBufferCopy"));
-					}
-					TransitionAndCopyTexture(RHICmdList, FinalBuffer, FinalBufferCopy->GetRHI(), {});
+			DrawElementsInputs.bElementsTextureIsHDRDisplay = false;
+			DrawElementsInputs.bAllowColorDeficiencyCorrection = true;
+		}
 
-					RHICmdList.Transition({
-						FRHITransitionInfo(FinalBuffer, ERHIAccess::Unknown, ERHIAccess::RTV)
-					});
-					FRHIRenderPassInfo RPInfo(FinalBuffer, ERenderTargetActions::Load_Store);
-					RHICmdList.BeginRenderPass(RPInfo, TEXT("SlateComposite"));
-					{
-						FGraphicsPipelineStateInitializer GraphicsPSOInit;
-						RHICmdList.ApplyCachedRenderTargets(GraphicsPSOInit);
-						GraphicsPSOInit.BlendState = TStaticBlendState<>::GetRHI();
-						GraphicsPSOInit.RasterizerState = TStaticRasterizerState<>::GetRHI();
-						GraphicsPSOInit.DepthStencilState = TStaticDepthStencilState<false, CF_Always>::GetRHI();
-
-						TShaderMapRef<FScreenVS> VertexShader(ShaderMap);
-
-						TShaderRef<FCompositePSBase> PixelShader;
-						if (bUseScRGB)
-						{
-							if (bApplyColorDeficiencyFilter)
-							{
-								PixelShader = TShaderMapRef<FCompositePS_SCRGB_CDV>(ShaderMap);
-							}
-							else
-							{
-								PixelShader = TShaderMapRef<FCompositePS_SCRGB>(ShaderMap);
-							}
-						}
-						else
-						{
-							if (bApplyColorDeficiencyFilter)
-							{
-								PixelShader = TShaderMapRef<FCompositePS_PQ_CDV>(ShaderMap);
-							}
-							else
-							{
-								PixelShader = TShaderMapRef<FCompositePS_PQ>(ShaderMap);
-							}
-						}
-
-						GraphicsPSOInit.BoundShaderState.VertexDeclarationRHI = GFilterVertexDeclaration.VertexDeclarationRHI;
-						GraphicsPSOInit.BoundShaderState.VertexShaderRHI = VertexShader.GetVertexShader();
-						GraphicsPSOInit.BoundShaderState.PixelShaderRHI = PixelShader.GetPixelShader();
-						GraphicsPSOInit.PrimitiveType = PT_TriangleList;
-
-						SetGraphicsPipelineState(RHICmdList, GraphicsPSOInit, 0);
-
-						SetShaderParametersLegacyPS(RHICmdList, PixelShader, 
-							ViewportInfo.UITargetRT->GetRHI(), UITargetRTMaskTexture, FinalBufferCopy->GetRHI(), ViewportInfo.ColorSpaceLUT,
-							GSlateColorDeficiencyCorrection, GSlateColorDeficiencyType, GSlateColorDeficiencySeverity, GSlateShowColorDeficiencyCorrectionWithDeficiency);
-
-						RendererModule.DrawRectangle(
-							RHICmdList,
-							0.f, 0.f,
-							(float)ViewportWidth, (float)ViewportHeight,
-							0.f, 0.f,
-							(float)ViewportWidth, (float)ViewportHeight,
-							FIntPoint(ViewportWidth, ViewportHeight),
-							FIntPoint(ViewportWidth, ViewportHeight),
-							VertexShader,
-							EDRF_UseTriangleOptimization);
-					}
-					RHICmdList.EndRenderPass();
-				}
-				else
-				{
-					TShaderRef<FCompositeCSBase> ComputeShader;
-					if (bUseScRGB)
-					{
-						if (bApplyColorDeficiencyFilter)
-						{
-							ComputeShader = TShaderMapRef<FCompositeCS_SCRGB_CDV>(ShaderMap);
-						}
-						else
-						{
-							ComputeShader = TShaderMapRef<FCompositeCS_SCRGB>(ShaderMap);
-						}
-					}
-					else
-					{
-						if (bApplyColorDeficiencyFilter)
-						{
-							ComputeShader = TShaderMapRef<FCompositeCS_PQ_CDV>(ShaderMap);
-						}
-						else
-						{
-							ComputeShader = TShaderMapRef<FCompositeCS_PQ>(ShaderMap);
-						}
-					}
+		DrawElementsInputs.ElementsTexture = ElementsTexture;
+		DrawElementsInputs.ElementsLoadAction = bClearElementsTexture ? ERenderTargetLoadAction::EClear : ERenderTargetLoadAction::ELoad;
+		DrawElementsInputs.ElementsBuffers = SlateElementsBuffers;
+		DrawElementsInputs.StencilTexture = BatchData.IsStencilClippingRequired() ? SlateStencilTexture : nullptr;
 
-					RHICmdList.Transition({
-						FRHITransitionInfo(FinalBuffer, ERHIAccess::Unknown, ERHIAccess::UAVCompute)
-					});
+		AddSlateDrawElementsPass(GraphBuilder, *RenderingPolicy, DrawElementsInputs, BatchData.GetRenderBatches(), BatchData.GetFirstRenderBatchIndex());
 
-					FVector4f SceneTextureDimensions((float)ViewportWidth, (float)ViewportHeight, 1.0f/(float)ViewportWidth, 1.0f/(float)ViewportHeight);
-					SetComputePipelineState(RHICmdList, ComputeShader.GetComputeShader());
+		if (bCompositeUIWithSceneHDR)
+		{
+			RDG_EVENT_SCOPE(GraphBuilder, "CompositeUI");
 
-					SetShaderParametersLegacyCS(RHICmdList, ComputeShader,
-						ViewportInfo.UITargetRT->GetRHI(), UITargetRTMaskTexture, BackBufferUAV, ViewportInfo.ColorSpaceLUT, SceneTextureDimensions,
-						GSlateColorDeficiencyCorrection, GSlateColorDeficiencyType, GSlateColorDeficiencySeverity, GSlateShowColorDeficiencyCorrectionWithDeficiency);
+			FRDGTexture* ElementsWriteMaskTexture = nullptr;
 
-					RHICmdList.DispatchComputeShader(FMath::DivideAndRoundUp(ViewportWidth, FCompositeCSBase::NUM_THREADS_PER_GROUP), FMath::DivideAndRoundUp(ViewportHeight, FCompositeCSBase::NUM_THREADS_PER_GROUP), 1);
+			if (RHISupportsRenderTargetWriteMask(GMaxRHIShaderPlatform))
+			{
+				FRenderTargetWriteMask::Decode(GraphBuilder, ShaderMap, MakeArrayView({ ElementsTexture }), ElementsWriteMaskTexture, ETextureCreateFlags::None, TEXT("ElementsWriteMaskTexture"));
+			}
 
-				}
+			static const auto CVarOutputDevice = IConsoleManager::Get().FindTConsoleVariableDataInt(TEXT("r.HDR.Display.OutputDevice"));
 
-				// Put the backbuffer back to the correct one.
-				BackBuffer = FinalBuffer;
-			} //bCompositeUI
+			FCompositeShader::FPermutationDomain PermutationVector;
+			PermutationVector.Set<FCompositeShader::FSCRGBEncoding>(ViewportInfo.HDRDisplayOutputFormat == EDisplayOutputFormat::HDR_ACES_1000nit_ScRGB || ViewportInfo.HDRDisplayOutputFormat == EDisplayOutputFormat::HDR_ACES_2000nit_ScRGB);
+			PermutationVector.Set<FCompositeShader::FApplyColorDeficiency>(GSlateColorDeficiencyType != EColorVisionDeficiency::NormalVision && GSlateColorDeficiencySeverity > 0);
 
+			FCompositeShaderCommonParameters CommonParameters;
+			CommonParameters.UIWriteMaskTexture = ElementsWriteMaskTexture;
+			CommonParameters.UITexture = ElementsTexture;
+			CommonParameters.UISampler = TStaticSamplerState<SF_Point>::GetRHI();
+			CommonParameters.UILevel = CVarUILevel.GetValueOnRenderThread();
+			CommonParameters.UILuminance = CVarHDRUILuminance.GetValueOnRenderThread();
+			CommonParameters.ColorVisionDeficiencySeverity = (float)GSlateColorDeficiencySeverity;
+			CommonParameters.ColorVisionDeficiencyType = (float)GSlateColorDeficiencyType;
+			CommonParameters.bSimulateCorrectionWithDeficiency = GSlateShowColorDeficiencyCorrectionWithDeficiency ? 1.0f : 0.0f;
+			CommonParameters.bCorrectDeficiency = GSlateColorDeficiencyCorrection ? 1.0f : 0.0f;
 
-#if WITH_EDITOR
-			if (bRenderOffscreen)
+			if (FCompositeCS::IsShaderSupported(GMaxRHIShaderPlatform))
 			{
-				const auto FeatureLevel = GMaxRHIFeatureLevel;
-				auto ShaderMap = GetGlobalShaderMap(FeatureLevel);
-
-				FRHIRenderPassInfo RPInfo(FinalBuffer, ERenderTargetActions::Load_Store);
-				RHICmdList.Transition(FRHITransitionInfo(FinalBuffer, ERHIAccess::Unknown, ERHIAccess::RTV));
-				RHICmdList.BeginRenderPass(RPInfo, TEXT("SlateComposite"));
-
-				FGraphicsPipelineStateInitializer GraphicsPSOInit;
-				RHICmdList.ApplyCachedRenderTargets(GraphicsPSOInit);
-				GraphicsPSOInit.BlendState			= TStaticBlendState<>::GetRHI();
-				GraphicsPSOInit.RasterizerState		= TStaticRasterizerState<>::GetRHI();
-				GraphicsPSOInit.DepthStencilState	= TStaticDepthStencilState<false, CF_Always>::GetRHI();
-
-				// ST2084 (PQ) encoding
-				TShaderMapRef<FHDREditorConvertPS> PixelShader(ShaderMap);
-				TShaderMapRef<FScreenVS> VertexShader(ShaderMap);
-
-				GraphicsPSOInit.BoundShaderState.VertexDeclarationRHI = GFilterVertexDeclaration.VertexDeclarationRHI;
-				GraphicsPSOInit.BoundShaderState.VertexShaderRHI = VertexShader.GetVertexShader();
-				GraphicsPSOInit.BoundShaderState.PixelShaderRHI = PixelShader.GetPixelShader();
-				GraphicsPSOInit.PrimitiveType = PT_TriangleList;
-
-				SetGraphicsPipelineState(RHICmdList, GraphicsPSOInit, 0);
-
-				SetShaderParametersLegacyPS(RHICmdList, PixelShader, HDRRenderRT->GetRHI());
-
-				static const FName RendererModuleName("Renderer");
-				IRendererModule& RendererModule = FModuleManager::GetModuleChecked<IRendererModule>(RendererModuleName);
-
-				RendererModule.DrawRectangle(
-					RHICmdList,
-					0.f, 0.f,
-					(float)ViewportWidth, (float)ViewportHeight,
-					0.f, 0.f,
-					(float)ViewportWidth, (float)ViewportHeight,
-					FIntPoint(ViewportWidth, ViewportHeight),
-					FIntPoint(ViewportWidth, ViewportHeight),
-					VertexShader,
-					EDRF_UseTriangleOptimization);
-				
-				RHICmdList.EndRenderPass();
-				BackBuffer = FinalBuffer;
+				FCompositeCS::FParameters* PassParameters = GraphBuilder.AllocParameters<FCompositeCS::FParameters>();
+				PassParameters->Common = CommonParameters;
+				PassParameters->RWSceneTexture = GraphBuilder.CreateUAV(OutputTexture);
+				PassParameters->SceneTextureDimensions = FVector4f((float)OutputExtent.X, (float)OutputExtent.Y, 1.0f/(float)OutputExtent.X, 1.0f/(float)OutputExtent.Y);
+
+				TShaderMapRef<FCompositeCS> ComputeShader(ShaderMap, PermutationVector);
+
+				FComputeShaderUtils::AddPass(
+					GraphBuilder,
+					RDG_EVENT_NAME("CompositeUI"),
+					ComputeShader,
+					PassParameters,
+					FIntVector(
+						FMath::DivideAndRoundUp<int32>(OutputExtent.X, FCompositeCS::NUM_THREADS_PER_GROUP),
+						FMath::DivideAndRoundUp<int32>(OutputExtent.Y, FCompositeCS::NUM_THREADS_PER_GROUP),
+						1));
 			}
-#endif
-			if (!bRenderedStereo && GEngine && IsValidRef(ViewportInfo.GetRenderTargetTexture()) && GEngine->StereoRenderingDevice.IsValid())
+			else
 			{
-				const FVector2D WindowSize = WindowElementList.GetWindowSize();
-				GEngine->StereoRenderingDevice->RenderTexture_RenderThread(RHICmdList, RHIGetViewportBackBuffer(ViewportInfo.ViewportRHI), ViewportInfo.GetRenderTargetTexture(), WindowSize);
-			}
-			RHICmdList.Transition(FRHITransitionInfo(BackBuffer, ERHIAccess::Unknown, ERHIAccess::SRVGraphics));
+				FRDGTexture* ViewportCopyTexture = GraphBuilder.CreateTexture(
+					FRDGTextureDesc::Create2D(
+						OutputExtent,
+						OutputTexture->Desc.Format,
+						FClearValueBinding::Transparent,
+						GetSlateTransientRenderTargetFlags()),
+					TEXT("SlateViewportCopyTexture"));
 
-			// Fire delegate to inform bound functions the back buffer is ready to be captured.
-			OnBackBufferReadyToPresentDelegate.Broadcast(*DrawCommandParams.Window, BackBuffer);
+				AddCopyTexturePass(GraphBuilder, OutputTexture, ViewportCopyTexture);
+					
+				const FScreenPassTextureViewport Viewport(OutputTexture);
 
-			if (bRenderedStereo)
-			{
-				// XR requires the backbuffer to be transitioned back to RTV
-				RHICmdList.Transition(FRHITransitionInfo(BackBuffer, ERHIAccess::SRVGraphics, ERHIAccess::RTV));
-			}
-		}
-	}
+				FCompositePS::FParameters* PassParameters = GraphBuilder.AllocParameters<FCompositePS::FParameters>();
+				PassParameters->Common = CommonParameters;
+				PassParameters->RenderTargets[0] = FRenderTargetBinding(OutputTexture, ERenderTargetLoadAction::ENoAction);
+				PassParameters->SceneTexture = ViewportCopyTexture;
+				PassParameters->SceneSampler = TStaticSamplerState<SF_Point>::GetRHI();
 
-	if (bTakingAScreenShot && ScreenshotViewportInfo != nullptr && ScreenshotViewportInfo == &ViewportInfo)
-	{
-		// take screenshot before swapbuffer
-		FTextureRHIRef BackBuffer = RHIGetViewportBackBuffer(ViewportInfo.ViewportRHI);
+				TShaderMapRef<FCompositePS> PixelShader(ShaderMap, PermutationVector);
 
-		// Sanity check to make sure the user specified a valid screenshot rect.
-		FIntRect ClampedScreenshotRect;
-		ClampedScreenshotRect.Min = ScreenshotRect.Min;
-		ClampedScreenshotRect.Max = ScreenshotRect.Max.ComponentMin(BackBuffer->GetSizeXY());
-		ClampedScreenshotRect.Max = ScreenshotRect.Min.ComponentMax(ClampedScreenshotRect.Max);
+				AddDrawScreenPass(GraphBuilder, RDG_EVENT_NAME("CompositeUI"), FScreenPassViewInfo(), Viewport, Viewport, PixelShader, PassParameters);
+			}
+		}
 
-		if (ClampedScreenshotRect != ScreenshotRect)
+		AddPass(GraphBuilder, RDG_EVENT_NAME("ReadyToPresent"), [this, Window = Inputs.Window, WindowSize = WindowElementList.GetWindowSize(), bCompositeStereoToSwapChain, SwapChainTextureRHI, ViewportTextureRHI, OutputTextureRHI] (FRHICommandListImmediate& RHICmdList)
 		{
-			UE_LOG(LogSlate, Warning, TEXT("Slate: Screenshot rect max coordinate had to be clamped from [%d, %d] to [%d, %d]"), ScreenshotRect.Max.X, ScreenshotRect.Max.Y, ClampedScreenshotRect.Max.X, ClampedScreenshotRect.Max.Y);
-		}
+			if (bCompositeStereoToSwapChain)
+			{
+				GEngine->StereoRenderingDevice->RenderTexture_RenderThread(RHICmdList, SwapChainTextureRHI, ViewportTextureRHI, WindowSize);
+			}
+			
+			// Fire delegate to inform bound functions the back buffer is ready to be captured.
+			OnBackBufferReadyToPresentDelegate.Broadcast(*Window, OutputTextureRHI);
+		});
 
-		if (!ClampedScreenshotRect.IsEmpty())
+		if (ScreenshotState.ViewportToCapture == &ViewportInfo)
 		{
-			if (OutHDRScreenshotData != nullptr)
+			// Sanity check to make sure the user specified a valid screenshot rect.
+			FIntRect ViewRectClamped;
+			ViewRectClamped.Min = ScreenshotState.ViewRect.Min;
+			ViewRectClamped.Max = ScreenshotState.ViewRect.Max.ComponentMin(OutputExtent);
+			ViewRectClamped.Max = ScreenshotState.ViewRect.Min.ComponentMax(ViewRectClamped.Max);
+
+			if (ViewRectClamped != ScreenshotState.ViewRect)
 			{
-				RHICmdList.ReadSurfaceData(BackBuffer, ClampedScreenshotRect, *OutHDRScreenshotData, FReadSurfaceDataFlags(RCM_MinMax));
+				UE_LOG(LogSlate, Warning, TEXT("Slate: Screenshot rect max coordinate had to be clamped from [%d, %d] to [%d, %d]"), ScreenshotState.ViewRect.Max.X, ScreenshotState.ViewRect.Max.Y, ViewRectClamped.Max.X, ViewRectClamped.Max.Y);
 			}
-			else
+
+			if (!ViewRectClamped.IsEmpty())
 			{
-				RHICmdList.ReadSurfaceData(BackBuffer, ClampedScreenshotRect, *OutScreenshotData, FReadSurfaceDataFlags());
+				AddReadbackTexturePass(GraphBuilder, RDG_EVENT_NAME("ScreenshotReadback"), OutputTexture,
+					[this, OutputTexture, ViewRectClamped, ColorDataHDR = ScreenshotState.ColorDataHDR, ColorData = ScreenshotState.ColorData] (FRHICommandListImmediate& RHICmdList)
+				{
+					if (ColorDataHDR)
+					{
+						RHICmdList.ReadSurfaceData(OutputTexture->GetRHI(), ViewRectClamped, *ColorDataHDR, FReadSurfaceDataFlags(RCM_MinMax));
+					}
+					else
+					{
+						check(ColorData);
+						RHICmdList.ReadSurfaceData(OutputTexture->GetRHI(), ViewRectClamped, *ColorData, FReadSurfaceDataFlags());
+					}
+				});
 			}
 		}
-		else
-		{
-			UE_LOG(LogSlate, Warning, TEXT("Slate: Screenshot rect was empty! Skipping readback of back buffer."));
-		}
-		bTakingAScreenShot = false;
-		OutScreenshotData = nullptr;
-		OutHDRScreenshotData = nullptr;
-		ScreenshotViewportInfo = nullptr;
 	}
 
-	// check if we need to cleanup slate render targets alloc.
-	RenderingPolicy->TickPostProcessResources();
-
-	// Calculate renderthread time (excluding idle time).	
-	uint32 StartTime = FPlatformTime::Cycles();
-
-	RHICmdList.EnqueueLambda([CurrentFrameCounter = GFrameCounterRenderThread](FRHICommandListImmediate& InRHICmdList)
+	AddPass(GraphBuilder, RDG_EVENT_NAME("EndDrawingViewport"), [this, ViewportRHI = ViewportInfo.ViewportRHI, bLockToVsync = Inputs.bLockToVsync] (FRHICommandListImmediate& RHICmdList)
 	{
-		UEngine::SetPresentLatencyMarkerStart(CurrentFrameCounter);
-	});
+		uint32 StartTime = FPlatformTime::Cycles();
 
-	RHICmdList.EndDrawingViewport(ViewportInfo.ViewportRHI, true, DrawCommandParams.bLockToVsync);
+		RHICmdList.EnqueueLambda([CurrentFrameCounter = GFrameCounterRenderThread](FRHICommandListImmediate& InRHICmdList)
+		{
+			UEngine::SetPresentLatencyMarkerStart(CurrentFrameCounter);
+		});
 
-	RHICmdList.EnqueueLambda([CurrentFrameCounter = GFrameCounterRenderThread](FRHICommandListImmediate& InRHICmdList)
-	{
-		UEngine::SetPresentLatencyMarkerEnd(CurrentFrameCounter);
-	});
+		RHICmdList.EndDrawingViewport(ViewportRHI, true, bLockToVsync);
 
-	uint32 EndTime = FPlatformTime::Cycles();
+		RHICmdList.EnqueueLambda([CurrentFrameCounter = GFrameCounterRenderThread](FRHICommandListImmediate& InRHICmdList)
+		{
+			UEngine::SetPresentLatencyMarkerEnd(CurrentFrameCounter);
+		});
 
-	GSwapBufferTime = EndTime - StartTime;
-	SET_CYCLE_COUNTER(STAT_PresentTime, GSwapBufferTime);
+		uint32 EndTime = FPlatformTime::Cycles();
 
-	uint32 ThreadTime = EndTime - LastTimestamp;
-	LastTimestamp = EndTime;
+		GSwapBufferTime = EndTime - StartTime;
+		SET_CYCLE_COUNTER(STAT_PresentTime, GSwapBufferTime);
 
-	uint32 RenderThreadIdle = 0;
+		static uint32 LastTimestamp = FPlatformTime::Cycles();
+		uint32 ThreadTime = EndTime - LastTimestamp;
+		LastTimestamp = EndTime;
 
-	FThreadIdleStats& RenderThread = FThreadIdleStats::Get();
-	GRenderThreadIdle[ERenderThreadIdleTypes::WaitingForAllOtherSleep] = RenderThread.Waits;
-	GRenderThreadIdle[ERenderThreadIdleTypes::WaitingForGPUPresent] += GSwapBufferTime;
+		uint32 RenderThreadIdle = 0;
 
-	SET_CYCLE_COUNTER(STAT_RenderingIdleTime_RenderThreadSleepTime, GRenderThreadIdle[ERenderThreadIdleTypes::WaitingForAllOtherSleep]);
-	SET_CYCLE_COUNTER(STAT_RenderingIdleTime_WaitingForGPUQuery   , GRenderThreadIdle[ERenderThreadIdleTypes::WaitingForGPUQuery     ]);
-	SET_CYCLE_COUNTER(STAT_RenderingIdleTime_WaitingForGPUPresent , GRenderThreadIdle[ERenderThreadIdleTypes::WaitingForGPUPresent   ]);
+		FThreadIdleStats& RenderThread = FThreadIdleStats::Get();
+		GRenderThreadIdle[ERenderThreadIdleTypes::WaitingForAllOtherSleep] = RenderThread.Waits;
+		GRenderThreadIdle[ERenderThreadIdleTypes::WaitingForGPUPresent] += GSwapBufferTime;
 
-	const uint32 RenderThreadNonCriticalWaits = RenderThread.Waits - RenderThread.WaitsCriticalPath;
-	const uint32 RenderThreadWaitingForGPUQuery = GRenderThreadIdle[ERenderThreadIdleTypes::WaitingForGPUQuery];
+		SET_CYCLE_COUNTER(STAT_RenderingIdleTime_RenderThreadSleepTime, GRenderThreadIdle[ERenderThreadIdleTypes::WaitingForAllOtherSleep]);
+		SET_CYCLE_COUNTER(STAT_RenderingIdleTime_WaitingForGPUQuery   , GRenderThreadIdle[ERenderThreadIdleTypes::WaitingForGPUQuery     ]);
+		SET_CYCLE_COUNTER(STAT_RenderingIdleTime_WaitingForGPUPresent , GRenderThreadIdle[ERenderThreadIdleTypes::WaitingForGPUPresent   ]);
 
-	// Set the RenderThreadIdle CSV stats
-	CSV_CUSTOM_STAT(RenderThreadIdle, Total          , FPlatformTime::ToMilliseconds(RenderThread.Waits            ), ECsvCustomStatOp::Set);
-	CSV_CUSTOM_STAT(RenderThreadIdle, CriticalPath   , FPlatformTime::ToMilliseconds(RenderThread.WaitsCriticalPath), ECsvCustomStatOp::Set);
-	CSV_CUSTOM_STAT(RenderThreadIdle, SwapBuffer     , FPlatformTime::ToMilliseconds(GSwapBufferTime               ), ECsvCustomStatOp::Set);
-	CSV_CUSTOM_STAT(RenderThreadIdle, NonCriticalPath, FPlatformTime::ToMilliseconds(RenderThreadNonCriticalWaits  ), ECsvCustomStatOp::Set);
-	CSV_CUSTOM_STAT(RenderThreadIdle, GPUQuery       , FPlatformTime::ToMilliseconds(RenderThreadWaitingForGPUQuery), ECsvCustomStatOp::Set);
+		const uint32 RenderThreadNonCriticalWaits = RenderThread.Waits - RenderThread.WaitsCriticalPath;
+		const uint32 RenderThreadWaitingForGPUQuery = GRenderThreadIdle[ERenderThreadIdleTypes::WaitingForGPUQuery];
 
-	for (int32 Index = 0; Index < ERenderThreadIdleTypes::Num; Index++)
-	{
-		RenderThreadIdle += GRenderThreadIdle[Index];
-		GRenderThreadIdle[Index] = 0;
-	}
+		// Set the RenderThreadIdle CSV stats
+		CSV_CUSTOM_STAT(RenderThreadIdle, Total          , FPlatformTime::ToMilliseconds(RenderThread.Waits            ), ECsvCustomStatOp::Set);
+		CSV_CUSTOM_STAT(RenderThreadIdle, CriticalPath   , FPlatformTime::ToMilliseconds(RenderThread.WaitsCriticalPath), ECsvCustomStatOp::Set);
+		CSV_CUSTOM_STAT(RenderThreadIdle, SwapBuffer     , FPlatformTime::ToMilliseconds(GSwapBufferTime               ), ECsvCustomStatOp::Set);
+		CSV_CUSTOM_STAT(RenderThreadIdle, NonCriticalPath, FPlatformTime::ToMilliseconds(RenderThreadNonCriticalWaits  ), ECsvCustomStatOp::Set);
+		CSV_CUSTOM_STAT(RenderThreadIdle, GPUQuery       , FPlatformTime::ToMilliseconds(RenderThreadWaitingForGPUQuery), ECsvCustomStatOp::Set);
 
-	SET_CYCLE_COUNTER(STAT_RenderingIdleTime, RenderThreadIdle);
-	GRenderThreadTime = (ThreadTime > RenderThreadIdle) ? (ThreadTime - RenderThreadIdle) : ThreadTime;
-	GRenderThreadWaitTime = RenderThreadIdle;
+		for (int32 Index = 0; Index < ERenderThreadIdleTypes::Num; Index++)
+		{
+			RenderThreadIdle += GRenderThreadIdle[Index];
+			GRenderThreadIdle[Index] = 0;
+		}
 
-	// Compute GRenderThreadTimeCriticalPath
-	uint32 RenderThreadNonCriticalPathIdle = RenderThreadIdle - RenderThread.WaitsCriticalPath;
-	GRenderThreadTimeCriticalPath = (ThreadTime > RenderThreadNonCriticalPathIdle) ? (ThreadTime - RenderThreadNonCriticalPathIdle) : ThreadTime;
-	SET_CYCLE_COUNTER(STAT_RenderThreadCriticalPath, GRenderThreadTimeCriticalPath);
+		SET_CYCLE_COUNTER(STAT_RenderingIdleTime, RenderThreadIdle);
+		GRenderThreadTime = (ThreadTime > RenderThreadIdle) ? (ThreadTime - RenderThreadIdle) : ThreadTime;
+		GRenderThreadWaitTime = RenderThreadIdle;
 
+		// Compute GRenderThreadTimeCriticalPath
+		uint32 RenderThreadNonCriticalPathIdle = RenderThreadIdle - RenderThread.WaitsCriticalPath;
+		GRenderThreadTimeCriticalPath = (ThreadTime > RenderThreadNonCriticalPathIdle) ? (ThreadTime - RenderThreadNonCriticalPathIdle) : ThreadTime;
+		SET_CYCLE_COUNTER(STAT_RenderThreadCriticalPath, GRenderThreadTimeCriticalPath);
 
-	if (CVarRenderThreadTimeIncludesDependentWaits.GetValueOnRenderThread())
-	{
-		// Optionally force the renderthread stat to include dependent waits
-		GRenderThreadTime = GRenderThreadTimeCriticalPath;
-	}
+		if (CVarRenderThreadTimeIncludesDependentWaits.GetValueOnRenderThread())
+		{
+			// Optionally force the renderthread stat to include dependent waits
+			GRenderThreadTime = GRenderThreadTimeCriticalPath;
+		}
 
-	// Reset the idle stats
-	RenderThread.Reset();
+		// Reset the idle stats
+		RenderThread.Reset();
 	
-	static TOptional<uint32> RHITCycles;
-	if (IsRunningRHIInSeparateThread())
-	{
-		RHICmdList.EnqueueLambda([](FRHICommandListImmediate&)
+		static TOptional<uint32> RHITCycles;
+		if (IsRunningRHIInSeparateThread())
 		{
-			// Update RHI thread time
-			FThreadIdleStats& RHIThreadStats = FThreadIdleStats::Get();
-
-			if (!RHITCycles.IsSet())
+			RHICmdList.EnqueueLambda([](FRHICommandListImmediate&)
 			{
-				RHITCycles = FPlatformTime::Cycles();
-			}
+				// Update RHI thread time
+				FThreadIdleStats& RHIThreadStats = FThreadIdleStats::Get();
+
+				if (!RHITCycles.IsSet())
+				{
+					RHITCycles = FPlatformTime::Cycles();
+				}
 
-			uint32 Next = FPlatformTime::Cycles();
+				uint32 Next = FPlatformTime::Cycles();
 
-			int32 Result = int32(Next - RHITCycles.GetValue() - RHIThreadStats.Waits);
-			RHITCycles = Next;
+				int32 Result = int32(Next - RHITCycles.GetValue() - RHIThreadStats.Waits);
+				RHITCycles = Next;
 
-			FPlatformAtomics::AtomicStore((int32*)&GRHIThreadTime, FMath::Max(Result, 0));
-			RHIThreadStats.Reset();
-		});
-	}
-	else
-	{
-		RHITCycles.Reset();
-	}
+				FPlatformAtomics::AtomicStore((int32*)&GRHIThreadTime, FMath::Max(Result, 0));
+				RHIThreadStats.Reset();
+			});
+		}
+		else
+		{
+			RHITCycles.Reset();
+		}
+	});
 }
 
 void FSlateRHIRenderer::DrawWindows(FSlateDrawBuffer& WindowDrawBuffer)
@@ -1641,34 +1009,26 @@ void FSlateRHIRenderer::DrawWindows(FSlateDrawBuffer& WindowDrawBuffer)
 	DrawWindows_Private(WindowDrawBuffer);
 }
 
-
 void FSlateRHIRenderer::PrepareToTakeScreenshot(const FIntRect& Rect, TArray<FColor>* OutColorData, SWindow* InScreenshotWindow)
 {
 	check(OutColorData);
 
-	bTakingAScreenShot = true;
-	ScreenshotRect = Rect;
-	OutScreenshotData = OutColorData;
-	OutHDRScreenshotData = nullptr;
-	ScreenshotViewportInfo = *WindowToViewportInfo.Find(InScreenshotWindow);
+	ScreenshotState.ViewRect = Rect;
+	ScreenshotState.ViewportToCapture = WindowToViewportInfo.FindRef(InScreenshotWindow);
+	ScreenshotState.ColorData = OutColorData;
+	ScreenshotState.ColorDataHDR = nullptr;
 }
 
 void FSlateRHIRenderer::PrepareToTakeHDRScreenshot(const FIntRect& Rect, TArray<FLinearColor>* OutColorData, SWindow* InScreenshotWindow)
 {
 	check(OutColorData);
 
-	bTakingAScreenShot = true;
-	ScreenshotRect = Rect;
-	OutScreenshotData = nullptr;
-	OutHDRScreenshotData = OutColorData;
-	ScreenshotViewportInfo = *WindowToViewportInfo.Find(InScreenshotWindow);
+	ScreenshotState.ViewRect = Rect;
+	ScreenshotState.ViewportToCapture = WindowToViewportInfo.FindRef(InScreenshotWindow);
+	ScreenshotState.ColorData = nullptr;
+	ScreenshotState.ColorDataHDR = OutColorData;
 }
 
-/**
-* Creates necessary resources to render a window and sends draw commands to the rendering thread
-*
-* @param WindowDrawBuffer	The buffer containing elements to draw
-*/
 void FSlateRHIRenderer::DrawWindows_Private(FSlateDrawBuffer& WindowDrawBuffer)
 {
 	checkSlow(IsThreadSafeForSlateRendering());
@@ -1681,366 +1041,299 @@ void FSlateRHIRenderer::DrawWindows_Private(FSlateDrawBuffer& WindowDrawBuffer)
 		bUpdateHDRDisplayInformation = false;
 	}
 
-	FSlateRHIRenderingPolicy* Policy = RenderingPolicy.Get();
-	ENQUEUE_RENDER_COMMAND(SlateBeginDrawingWindowsCommand)(
-		[Policy](FRHICommandListImmediate& RHICmdList)
-	{
-		Policy->BeginDrawingWindows();
-	}
-	);
-
-	// Update texture atlases if needed and safe
 	if (DoesThreadOwnSlateRendering())
 	{
 		ResourceManager->UpdateTextureAtlases();
 	}
 
+	USlateRHIRendererSettings* RendererSettings = USlateRHIRendererSettings::GetMutable();
+
+	const float AppDeltaTime = FApp::GetDeltaTime();
+	const FGameTime AppDilatedTime = FGameTime::CreateDilated(FPlatformTime::Seconds() - GStartTime, AppDeltaTime, FApp::GetCurrentTime() - GStartTime, AppDeltaTime);
+	const FVector2f AppCursorPosition = FSlateApplication::Get().GetCursorPos();
+	const bool bAppCanRender = GIsClient && !IsRunningCommandlet() && !GUsingNullRHI;
+	const bool bAppCanRenderPostProcess = IsInGameThread() && UAssetManager::IsInitialized() && bAppCanRender && CVarCopyBackbufferToSlatePostRenderTargets.GetValueOnGameThread() > 0;
+	EPixelFormat AppViewportSceneFormat = PF_Unknown;
+	FIntPoint AppViewportExtentMax = FIntPoint::ZeroValue;
+	ESlatePostRT PostProcessAnyUsedBits = ESlatePostRT::None;
+
 	const TSharedRef<FSlateFontCache> FontCache = SlateFontServices->GetFontCache();
 
-	// Iterate through each element list and set up an RHI window for it if needed
-	const TArray<TSharedRef<FSlateWindowElementList>>& WindowElementLists = WindowDrawBuffer.GetWindowElementLists();
-	for (int32 ListIndex = 0; ListIndex < WindowElementLists.Num(); ++ListIndex)
+	struct FWindowToRender
 	{
-		FSlateWindowElementList& ElementList = *WindowElementLists[ListIndex];
+		SWindow* Window;
+		FSlateWindowElementList* WindowElementList;
+		FSlateViewportInfo* ViewportInfo;
+		FIntPoint ViewportOffset;
+		FIntPoint ViewportExtent;
+		FIntRect  ViewportRect;
+		float ViewportScaleUI;
+		ESlatePostRT PostProcessUsedBits;
+		ESlatePostRT PostProcessCustumDrawBits;
+		ESlatePostRT PostProcessSkipUpdateBits;
+		FIntPoint CursorPosition;
+		bool bLockToVsync;
+	};
 
-		SWindow* Window = ElementList.GetRenderWindow();
+	TArray<FWindowToRender, FConcurrentLinearArrayAllocator> WindowsToRender;
+	WindowsToRender.Reserve(WindowDrawBuffer.GetWindowElementLists().Num());
+
+	for (const TSharedRef<FSlateWindowElementList>& WindowElementListRef : WindowDrawBuffer.GetWindowElementLists())
+	{
+		FSlateWindowElementList* WindowElementList = &(*WindowElementListRef);
+		SWindow* Window = WindowElementList->GetRenderWindow();
 
-		if (Window)
+		if (!Window)
 		{
-			const FVector2f WindowSize = UE::Slate::CastToVector2f(Window->GetViewportSize());
-			if (WindowSize.X > 0 && WindowSize.Y > 0)
-			{
-				// The viewport need to be created at this point  
-				FViewportInfo* ViewInfo = nullptr;
-				{
-					FViewportInfo** FoundViewInfo = WindowToViewportInfo.Find(Window);
-					if (ensure(FoundViewInfo))
-					{
-						ViewInfo = *FoundViewInfo;
-					}
-					else
-					{
-						UE_LOG(LogSlate, Error, TEXT("The ViewportInfo could not be found for Window."));
-						continue;
-					}
-				}
+			ensureMsgf(false, TEXT("Window isn't valid but being drawn!"));
+			continue;
+		}
 
-				bool bCompositeHDRViewports = ElementBatcher->CompositeHDRViewports();
-				ElementBatcher->SetCompositeHDRViewports(CompositeUIWithHdrRenderTarget(ViewInfo));
-				// Add all elements for this window to the element batcher
-				ElementBatcher->AddElements(ElementList);
-				ElementBatcher->SetCompositeHDRViewports(bCompositeHDRViewports);
+		// This will return zero if both the viewport and the window are zero sized.
+		const FVector2f WindowSize = Window->GetViewportSize();
+		if (WindowSize.X <= 0.0f || WindowSize.Y <= 0.0f)
+		{
+			continue;
+		}
 
-				// Update the font cache with new text after elements are batched
-				FontCache->UpdateCache();
+		TRACE_CPUPROFILER_EVENT_SCOPE(GatherWindowElements);
 
-				bool bLockToVsync = ElementBatcher->RequiresVsync();
+		// It's possible for a window to not have a viewport, in which case the viewport dimensions will be zero.
+		FVector2f ViewportCursorPosition = AppCursorPosition - Window->GetPositionInScreen();
+		FIntPoint ViewportOffset = FIntPoint::ZeroValue;
+		FIntPoint ViewportExtent = FIntPoint::ZeroValue;
+		FIntRect ViewportRect;
+		float ViewportScaleUI = Window->GetViewportScaleUIOverride();
 
-				bool bForceVsyncFromCVar = false;
-				if (GIsEditor)
-				{
-					static IConsoleVariable* CVar = IConsoleManager::Get().FindConsoleVariable(TEXT("r.VSyncEditor"));
-					bForceVsyncFromCVar = (CVar->GetInt() != 0);
-				}
-				else
-				{
-					static IConsoleVariable* CVar = IConsoleManager::Get().FindConsoleVariable(TEXT("r.VSync"));
-					bForceVsyncFromCVar = (CVar->GetInt() != 0);
-				}
+		if (ISlateViewport* Viewport = Window->GetViewport().Get())
+		{
+			TSharedPtr<SWidget> ViewportWidget = Viewport->GetWidget().Pin();
 
-				bLockToVsync |= bForceVsyncFromCVar;
+			if (ViewportWidget)
+			{
+				ViewportOffset = FIntPoint(
+					FMath::RoundToInt32(ViewportWidget->GetTickSpaceGeometry().GetAbsolutePosition().X - Window->GetPositionInScreen().X),
+					FMath::RoundToInt32(ViewportWidget->GetTickSpaceGeometry().GetAbsolutePosition().Y - Window->GetPositionInScreen().Y));
 
-				// Cache if the element batcher post buffer usage, will get reset on batch reset
-				ESlatePostRT UsedSlatePostBufferBits = ElementBatcher->GetUsedSlatePostBuffers();
-				ESlatePostRT ResourceUpdatingPostBufferBits = ElementBatcher->GetResourceUpdatingPostBuffers();
-				ESlatePostRT SkipDefaultUpdatePostBufferBits = ElementBatcher->GetSkipDefaultUpdatePostBuffers();
+				ViewportCursorPosition -= ViewportWidget->GetPaintSpaceGeometry().AbsolutePosition;
+			}
 
-				// All elements for this window have been batched and rendering data updated
-				ElementBatcher->ResetBatches();
+			ViewportExtent = Viewport->GetSize();
+			ViewportRect = FIntRect(ViewportOffset, ViewportOffset + ViewportExtent);
 
-				// Cache off the HDR status
-				ViewInfo->bHDREnabled = RHIGetColorSpace(ViewInfo->ViewportRHI) != EColorSpaceAndEOTF::ERec709_sRGB;
-				Window->SetIsHDR(ViewInfo->bHDREnabled);
+			ensureMsgf(AppViewportSceneFormat == PF_Unknown || AppViewportSceneFormat == Viewport->GetSceneTargetFormat(),
+				TEXT("Multiple viewport formats coming from multiple windows are not a supported scenario in slate. This will cause undefined behavior with Slate Post Buffers."));
 
-				if (Window->IsViewportSizeDrivenByWindow())
-				{
-					// Resize the viewport if needed
-					ConditionalResizeViewport(ViewInfo, ViewInfo->DesiredWidth, ViewInfo->DesiredHeight, IsViewportFullscreen(*Window), Window);
-				}
+			AppViewportSceneFormat = Viewport->GetSceneTargetFormat();
+			AppViewportExtentMax = AppViewportExtentMax.ComponentMax(ViewportExtent);
 
-				// Update slate post buffers before slate draw if needed
-				if (CVarCopyBackbufferToSlatePostRenderTargets.GetValueOnAnyThread() && IsInGameThread() && GIsClient && !IsRunningCommandlet() && !GUsingNullRHI && UAssetManager::IsInitialized())
-				{
-					uint8 SlatePostBufferBitIndex = 0;
-					for (ESlatePostRT SlatePostBufferBit : TEnumRange<ESlatePostRT>())
-					{
-						// We only attempt to load if the buffer is enabled, so just try to load / get the buffer
-						UTextureRenderTarget2D* SlatePostBuffer = USlateRHIRendererSettings::GetMutable()->LoadGetPostBufferRT(SlatePostBufferBit);
-						if (!SlatePostBuffer)
-						{
-							SlatePostBufferBitIndex++;
-							continue;
-						}
-
-						bool bIsViewportPresentForPIE = GIsEditor ? Window->GetViewport().IsValid() : true;
-						bool bPostBufferBitUsed = (UsedSlatePostBufferBits & SlatePostBufferBit) != ESlatePostRT::None;
-						bool bSkipDefaultPostBufferUpdate = (SkipDefaultUpdatePostBufferBits & SlatePostBufferBit) != ESlatePostRT::None;
-
-						if (bIsViewportPresentForPIE && bPostBufferBitUsed)
-						{
-							// Viewport only has a RT texture in editor / PIE, this texture will be the size of the viewport
-							// typically this texture 'BufferedRT', is smaller than the entire backbuffer used for the editor
-							// We initialize to nullptr so this ptr can be copied & resolved here instead of in the render command
-							FSlateRenderTargetRHI* ViewportTexture = nullptr;
-							if (GIsEditor)
-							{
-								ViewportTexture = static_cast<FSlateRenderTargetRHI*>(Window->GetViewport()->GetViewportRenderTargetTexture());
-							}
-
-							FIntPoint SizeSlatePostRT = GIsEditor
-								? FIntPoint(ViewportTexture->GetWidth(), ViewportTexture->GetHeight())
-								: FIntPoint(ViewInfo->DesiredWidth, ViewInfo->DesiredHeight);
-
-							bool bHDREnabled = IsHDREnabled();
-							bool bIsPixelFormatCorrect = SlatePostBuffer->GetFormat() == Window->GetViewport()->GetSceneTargetFormat();
-							if (SlatePostBuffer->SizeX != SizeSlatePostRT.X || SlatePostBuffer->SizeY != SizeSlatePostRT.Y || !bIsPixelFormatCorrect)
-							{
-								SlatePostBuffer->InitCustomFormat(SizeSlatePostRT.X, SizeSlatePostRT.Y, Window->GetViewport()->GetSceneTargetFormat(), true);
-							}
-
-							const FVector2D ElementWindowSize = ElementList.GetWindowSize();
-
-							if (!bSkipDefaultPostBufferUpdate)
-							{
-								if (USlateRHIPostBufferProcessor* PostProcessor = USlateFXSubsystem::GetPostProcessor(SlatePostBufferBit))
-								{
-									// Allow the post processor to enque render commands, we delgate this task so it can be done in a thread safe manner.
-									PostProcessor->PostProcess(ViewInfo, ViewportTexture, ElementWindowSize, FSlateRHIRenderingPolicyInterface(RenderingPolicy.Get()), SlatePostBuffer);
-								}
-								else
-								{
-									ENQUEUE_RENDER_COMMAND(FUpdateSlatePostBuffers)([ViewInfo, ElementWindowSize, ViewportTexture, SlatePostBuffer](FRHICommandListImmediate& RHICmdList)
-									{
-										bool bRenderedStereo = false;
-										if (CVarDrawToVRRenderTarget->GetInt() == 0 && GEngine && IsValidRef(ViewInfo->GetRenderTargetTexture()) && GEngine->StereoRenderingDevice.IsValid())
-										{
-											GEngine->StereoRenderingDevice->RenderTexture_RenderThread(RHICmdList, RHIGetViewportBackBuffer(ViewInfo->ViewportRHI), ViewInfo->GetRenderTargetTexture(), ElementWindowSize);
-											bRenderedStereo = true;
-										}
-
-										FTextureRHIRef ViewportRT = bRenderedStereo ? nullptr : ViewInfo->GetRenderTargetTexture();
-										FTextureRHIRef BackBuffer = (ViewportRT) ? ViewportRT : RHIGetViewportBackBuffer(ViewInfo->ViewportRHI);
-
-										if (BackBuffer)
-										{
-											FRHICopyTextureInfo CopyInfo;
-
-											// Copy just the viewport RT if in PIE, else do entire backbuffer
-											if (GIsEditor)
-											{
-												CopyInfo.Size = FIntVector(ViewportTexture->GetWidth(), ViewportTexture->GetHeight(), 1);
-												TransitionAndCopyTexture(RHICmdList, ViewportTexture->GetRHIRef(), SlatePostBuffer->TextureReference.TextureReferenceRHI, CopyInfo);
-											}
-											else
-											{
-												FIntPoint BackbufferExtent = BackBuffer->GetDesc().Extent;
-												CopyInfo.Size = FIntVector(BackbufferExtent.X, BackbufferExtent.Y, 1);
-												TransitionAndCopyTexture(RHICmdList, BackBuffer, SlatePostBuffer->TextureReference.TextureReferenceRHI, CopyInfo);
-											}
-										}
-									});
-								}
-
-								SlatePostRTFences[SlatePostBufferBitIndex].BeginFence();
-								LastFramesPostBufferUsed[SlatePostBufferBitIndex] = GFrameCounter;
-							}
-
-							bShrinkPostBufferRequested &= ~SlatePostBufferBit;
-						}
-						else if (SlatePostBuffer 
-							&& SlatePostBuffer->GetResource()
-							&& LastFramesPostBufferUsed[SlatePostBufferBitIndex] + 1 < GFrameCounter
-							&& SlatePostRTFences[SlatePostBufferBitIndex].IsFenceComplete() 
-							&& (SlatePostBuffer->SizeX != 1 || SlatePostBuffer->SizeY != 1))
-						{
-							if ((bShrinkPostBufferRequested & SlatePostBufferBit) == ESlatePostRT::None)
-							{
-								// Delay shrink attempts for a frame, since if we resize while there is an active copy in flight, we will crash
-								bShrinkPostBufferRequested |= SlatePostBufferBit;
-							}
-							else
-							{
-								// Resize unused SlatePostRTs to 1x1.
-								auto GetBackBufferFormat = [&]() 
-								{
-									// This path can execute when stopping PIE / exiting, so viewport may not be valid (above Viewport is always valid).
-									if (TSharedPtr<ISlateViewport> Viewport = Window->GetViewport())
-									{
-										return Viewport->GetSceneTargetFormat();
-									}
-
-									return IsHDREnabled() ? EPixelFormat::PF_FloatRGBA : EPixelFormat::PF_A2B10G10R10;
-								};
-
-								SlatePostBuffer->InitCustomFormat(1, 1, GetBackBufferFormat(), true);
-								bShrinkPostBufferRequested &= ~SlatePostBufferBit;
-							}
-						}
-
-						SlatePostBufferBitIndex++;
-					}
-				}
+			if (GIsEditor && ViewportScaleUI < 0.0f)
+			{
+				ViewportScaleUI = GetDefault<UUserInterfaceSettings>()->GetDPIScaleBasedOnSize(ViewportExtent);
+			}
+		}
 
-				// Tell the rendering thread to draw the windows
-				{
-					auto GetViewRect = [Window]()
-					{
-#if WITH_EDITOR
-						if (GIsEditor)
-						{
-							if (TSharedPtr<ISlateViewport> Viewport = Window->GetViewport())
-							{
-								if (TSharedPtr<SWidget> ViewportWidget = Viewport->GetWidget().Pin())
-								{
-									// The actual backbuffer has a padding that extends beyond the draw area, account for this in our offsets
-									int32 OffsetX = FMath::RoundToInt32(ViewportWidget->GetTickSpaceGeometry().GetAbsolutePosition().X - Window->GetPositionInScreen().X);
-									int32 OffsetY = FMath::RoundToInt32(ViewportWidget->GetTickSpaceGeometry().GetAbsolutePosition().Y - Window->GetPositionInScreen().Y);
-
-									FIntPoint ViewportSize = Viewport->GetSize();
-									return FIntRect(FIntPoint(OffsetX, OffsetY), FIntPoint(ViewportSize.X + OffsetX, ViewportSize.Y + OffsetY));
-								}
-							}
-						}
-#endif // WITH_EDITOR
-
-						return FIntRect();
-					};
-
-					FSlateDrawWindowCommandParams Params;
-
-					FVector2f CursorPosition = FSlateApplication::Get().GetCursorPos();
-
-					CursorPosition -= Window->GetPositionInScreen();
-					if (TSharedPtr<ISlateViewport> Viewport = Window->GetViewport())
-					{
-						if (TSharedPtr<SWidget> ViewportWidget = Viewport->GetWidget().Pin())
-						{
-							CursorPosition -= ViewportWidget->GetPaintSpaceGeometry().AbsolutePosition;
-						}
-					}
+		FSlateViewportInfo* ViewportInfo = WindowToViewportInfo.FindRef(Window);
 
-					Params.Renderer = this;
-					Params.WindowElementList = &ElementList;
-					Params.Window = Window;
-					Params.ViewRect = GetViewRect();
-					Params.ViewportScaleUI = Window->GetViewportScaleUIOverride();
-					Params.CursorPostion = FIntPoint(CursorPosition.X, CursorPosition.Y);
-					Params.UsedSlatePostBuffers = UsedSlatePostBufferBits;
-#if WANTS_DRAW_MESH_EVENTS
-					Params.WindowTitle = Window->GetTitle().ToString();
-#endif
-					Params.bLockToVsync = bLockToVsync;
-#if ALPHA_BLENDED_WINDOWS
-					Params.bClear = Window->GetTransparencySupport() == EWindowTransparency::PerPixel;
-#else
-					Params.bClear = false;
-#endif	
-					Params.Time = FGameTime::CreateDilated(
-						FPlatformTime::Seconds() - GStartTime, (float)FApp::GetDeltaTime(),
-						FApp::GetCurrentTime() - GStartTime, (float)FApp::GetDeltaTime());
-
-					// Negative Viewport UI Scale, get the value from user settings
-					if (GIsEditor && Params.ViewportScaleUI < 0)
-					{	
-						Params.ViewportScaleUI = GetDefault<UUserInterfaceSettings>()->GetDPIScaleBasedOnSize(Params.ViewRect.Size());
-					}
+		if (Window->IsViewportSizeDrivenByWindow())
+		{
+			ResizeViewportIfNeeded(ViewportInfo, ViewportInfo->ExtentToResizeTo, IsViewportFullscreen(*Window), Window);
+		}
 
-					// Skip the actual draw if we're in a headless execution environment
-					bool bLocalTakingAScreenShot = bTakingAScreenShot;
-					if (GIsClient && !IsRunningCommandlet() && !GUsingNullRHI)
-					{
-						ENQUEUE_RENDER_COMMAND(SlateDrawWindowsCommand)(
-							[Params, ViewInfo](FRHICommandListImmediate& RHICmdList)
-							{
-								Params.Renderer->DrawWindow_RenderThread(RHICmdList, *ViewInfo, *Params.WindowElementList, Params);
-							}
-						);
-
-						// After we draw, if a resource is going to update a post buffer we need to add another render fence for that buffer
-						if (CVarCopyBackbufferToSlatePostRenderTargets.GetValueOnAnyThread())
-						{
-							uint8 SlatePostBufferBitIndex = 0;
-							for (ESlatePostRT SlatePostBufferBit : TEnumRange<ESlatePostRT>())
-							{
-								bool bResourceUpdatedPostBuffer = (ResourceUpdatingPostBufferBits & SlatePostBufferBit) != ESlatePostRT::None;
-								if (bResourceUpdatedPostBuffer)
-								{
-									SlatePostRTFences[SlatePostBufferBitIndex].BeginFence();
-									LastFramesPostBufferUsed[SlatePostBufferBitIndex] = GFrameCounter;
-									bShrinkPostBufferRequested &= ~SlatePostBufferBit;
-								}
-
-								SlatePostBufferBitIndex++;
-							}
-						}
-					}
+		Window->SetIsHDR(ViewportInfo->bDisplayFormatIsHDR);
+		Window->ResetViewportScaleUIOverride();
+
+		ElementBatcher->SetCompositeHDRViewports(ViewportInfo->bDisplayFormatIsHDR && CompositeUIWithSceneHDR());
+		ElementBatcher->AddElements(*WindowElementList);
+
+		const bool bWindowCanRenderPostProcess       = bAppCanRender && ViewportExtent != FIntPoint::ZeroValue;
+		const ESlatePostRT PostProcessUsedBits       = bWindowCanRenderPostProcess ? ElementBatcher->GetUsedSlatePostBuffers() : ESlatePostRT::None;
+		const ESlatePostRT PostProcessCustomDrawBits = bWindowCanRenderPostProcess ? ElementBatcher->GetResourceUpdatingPostBuffers() : ESlatePostRT::None;
+		const ESlatePostRT PostProcessSkipUpdateBits = bWindowCanRenderPostProcess ? ElementBatcher->GetSkipDefaultUpdatePostBuffers() : ESlatePostRT::None;
+		const bool bLockToVsync                      = IsVSyncRequired(*ElementBatcher);
+
+		ElementBatcher->ResetBatches();
+		ElementBatcher->SetCompositeHDRViewports(false);
+
+		WindowsToRender.Emplace(FWindowToRender
+		{
+			  .Window                    = Window
+			, .WindowElementList         = WindowElementList
+			, .ViewportInfo              = ViewportInfo
+			, .ViewportOffset            = ViewportOffset
+			, .ViewportExtent            = ViewportExtent
+			, .ViewportRect              = ViewportRect
+			, .ViewportScaleUI           = ViewportScaleUI
+			, .PostProcessUsedBits       = PostProcessUsedBits
+			, .PostProcessCustumDrawBits = PostProcessCustomDrawBits
+			, .PostProcessSkipUpdateBits = PostProcessSkipUpdateBits
+			, .CursorPosition            = FIntPoint(ViewportCursorPosition.X, ViewportCursorPosition.Y)
+			, .bLockToVsync              = bLockToVsync
+		});
 
-					SlateWindowRendered.Broadcast(*Window, &ViewInfo->ViewportRHI);
+		PostProcessAnyUsedBits |= PostProcessUsedBits;
+	}
 
-					// Now that all widgets have drawn, reset material UI scale
-					Window->ResetViewportScaleUIOverride();
+	// Update the font cache now that all element batches were processed.
+	FontCache->UpdateCache();
 
-					if (bLocalTakingAScreenShot)
+	// Allocate any post process render targets that are used by any viewport.
+	if (bAppCanRenderPostProcess)
+	{
+		if (AppViewportExtentMax.X != 0 && AppViewportExtentMax.Y != 0)
+		{
+			for (ESlatePostRT Bit : MakeFlagsRange(PostProcessAnyUsedBits))
+			{
+				if (UTextureRenderTarget2D* RenderTarget = RendererSettings->LoadGetPostBufferRT(Bit))
+				{
+					if (RenderTarget->SizeX != AppViewportExtentMax.X || RenderTarget->SizeY != AppViewportExtentMax.Y || RenderTarget->GetFormat() != AppViewportSceneFormat)
 					{
-						// immediately flush the RHI command list
-						FlushRenderingCommands();
+						TRACE_CPUPROFILER_EVENT_SCOPE(AllocatePostProcessTexture);
+						RenderTarget->InitCustomFormat(AppViewportExtentMax.X, AppViewportExtentMax.Y, AppViewportSceneFormat, true);
 					}
+					PostProcessRenderTargets.LastUsedFrameCounter[(int32)Bit] = GFrameCounter;
 				}
 			}
 		}
-		else
+
+		for (ESlatePostRT Bit : MakeFlagsRange(ESlatePostRT::All & ~PostProcessAnyUsedBits))
 		{
-			ensureMsgf(false, TEXT("Window isnt valid but being drawn!"));
+			UTextureRenderTarget2D* RenderTarget = RendererSettings->TryGetPostBufferRT(Bit);;
+
+			if (RenderTarget
+				&& RenderTarget->GetResource()
+				&& RenderTarget->SizeX != 1
+				&& RenderTarget->SizeY != 1
+				&& PostProcessRenderTargets.LastUsedFrameCounter[(int32)Bit] < GFrameCounter)
+			{
+				// Trim unused post process render targets down to 1x1 to reclaim memory.
+				TRACE_CPUPROFILER_EVENT_SCOPE(TrimPostProcessTexture);
+				RenderTarget->InitCustomFormat(1, 1, PF_A2B10G10R10, true);
+			}
 		}
 	}
 
-	FSlateDrawBuffer* DrawBuffer = &WindowDrawBuffer;
-	ENQUEUE_RENDER_COMMAND(SlateEndDrawingWindowsCommand)(
-		[DrawBuffer, Policy](FRHICommandListImmediate& RHICmdList)
+	struct FDrawWindowsCommand : public TConcurrentLinearObject<FDrawWindowsCommand>
 	{
-		FSlateEndDrawingWindowsCommand::EndDrawingWindows(RHICmdList, DrawBuffer, *Policy);
-	}
-	);
+		bool IsEmpty() { return Windows.IsEmpty() && DeferredUpdates.IsEmpty(); }
+
+		TArray<FSlateDrawWindowPassInputs, FConcurrentLinearArrayAllocator> Windows;
+		TArray<FSlatePostProcessUpdateRequest, FConcurrentLinearArrayAllocator> PostProcessUpdates;
+		TArray<FRenderThreadUpdateContext, FConcurrentLinearArrayAllocator> DeferredUpdates;
+	};
+
+	TUniquePtr<FDrawWindowsCommand> DrawWindowsCommand = MakeUnique<FDrawWindowsCommand>();
+	DrawWindowsCommand->Windows.Reserve(WindowsToRender.Num());
+	DrawWindowsCommand->DeferredUpdates = MoveTemp(DeferredUpdateContexts);
+	DrawWindowsCommand->PostProcessUpdates.Reserve(WindowsToRender.Num() * FMath::CountBits((int32)PostProcessAnyUsedBits));
 
-	if (DeferredUpdateContexts.Num() > 0)
+	bool bScreenshotProcessed = false;
+
+	for (const FWindowToRender& WindowToRender : WindowsToRender)
 	{
-		// Intentionally copy the contexts to avoid contention with the game thread
-		ENQUEUE_RENDER_COMMAND(DrawWidgetRendererImmediate)(
-			[Contexts = DeferredUpdateContexts](FRHICommandListImmediate& RHICmdList) mutable
+		const int32 PostProcessUpdatesOffset = DrawWindowsCommand->PostProcessUpdates.Num();
+
+		if (bAppCanRenderPostProcess)
+		{
+			// Process bits that were NOT marked to skip the update.
+			for (ESlatePostRT Bit : MakeFlagsRange(ESlatePostRT::All & WindowToRender.PostProcessUsedBits & ~WindowToRender.PostProcessSkipUpdateBits))
 			{
-				for (const FRenderThreadUpdateContext& Context : Contexts)
+				UTextureRenderTarget2D* RenderTarget = RendererSettings->TryGetPostBufferRT(Bit);
+				check(RenderTarget);
+
+				FSlatePostProcessUpdateRequest Request;
+				Request.RenderTarget = Bit;
+				Request.RenderTargetTextureResource = RenderTarget->GetResource();
+
+				if (USlateRHIPostBufferProcessor* PostProcessor = USlateFXSubsystem::GetPostProcessor(Bit))
 				{
-					Context.Renderer->DrawWindowToTarget_RenderThread(RHICmdList, Context);
-					Context.Renderer->ReleaseDrawBuffer(*Context.WindowDrawBuffer);
+					Request.PostProcessorProxy = PostProcessor->GetRenderThreadProxy();
 				}
+
+				checkf(DrawWindowsCommand->PostProcessUpdates.Num() != DrawWindowsCommand->PostProcessUpdates.Max(), TEXT("This container is about to resize which will result in a dangling memory access"));
+				DrawWindowsCommand->PostProcessUpdates.Add(Request);
+			}
+		}
+
+		TConstArrayView<FSlatePostProcessUpdateRequest> PostProcessUpdatesForWindow;
+
+		const int32 PostProcessUpdateRequestsCount = DrawWindowsCommand->PostProcessUpdates.Num() - PostProcessUpdatesOffset;
+
+		if (PostProcessUpdateRequestsCount != 0)
+		{
+			PostProcessUpdatesForWindow = MakeArrayView(&DrawWindowsCommand->PostProcessUpdates[PostProcessUpdatesOffset], PostProcessUpdateRequestsCount);
+		}
+
+		bScreenshotProcessed |= ScreenshotState.ViewportToCapture == WindowToRender.ViewportInfo;
+
+		if (bAppCanRender)
+		{
+			DrawWindowsCommand->Windows.Emplace(FSlateDrawWindowPassInputs
+			{
+				  .Renderer                  = this
+				, .WindowElementList         = WindowToRender.WindowElementList
+				, .Window                    = WindowToRender.Window
+				, .ViewportInfo              = WindowToRender.ViewportInfo
+				, .PostProcessUpdateRequests = PostProcessUpdatesForWindow
+				, .CursorPosition            = WindowToRender.CursorPosition
+				, .SceneViewRect             = WindowToRender.ViewportRect
+				, .ViewportScaleUI           = WindowToRender.ViewportScaleUI
+				, .UsedSlatePostBuffers      = WindowToRender.PostProcessUsedBits
+			#if WANTS_DRAW_MESH_EVENTS
+				, .WindowTitle               = WindowToRender.Window->GetTitle().ToString()
+			#endif
+				, .Time                      = AppDilatedTime
+				, .bLockToVsync              = WindowToRender.bLockToVsync
+			#if ALPHA_BLENDED_WINDOWS
+				, .bClear                    = WindowToRender.Window->GetTransparencySupport() == EWindowTransparency::PerPixel
+			#endif
+			});
+		}
+	}
+
+	if (!DrawWindowsCommand->IsEmpty())
+	{
+		ENQUEUE_RENDER_COMMAND(SlateDrawWindowsCommand)([DrawWindowsCommand = MoveTemp(DrawWindowsCommand)](FRHICommandListImmediate& RHICmdList)
+		{
+			FRDGBuilder GraphBuilder(RHICmdList, RDG_EVENT_NAME("Slate"));
+
+			for (const FRenderThreadUpdateContext& DeferredUpdateContext : DrawWindowsCommand->DeferredUpdates)
+			{
+				DeferredUpdateContext.Renderer->DrawWindowToTarget_RenderThread(GraphBuilder, DeferredUpdateContext);
+			}
+
+			for (const FSlateDrawWindowPassInputs& DrawWindowPassInputs : DrawWindowsCommand->Windows)
+			{
+				DrawWindowPassInputs.Renderer->DrawWindow_RenderThread(GraphBuilder, DrawWindowPassInputs);
+			}
+
+			GraphBuilder.Execute();
+
+			for (const FRenderThreadUpdateContext& DeferredUpdateContext : DrawWindowsCommand->DeferredUpdates)
+			{
+				DeferredUpdateContext.Renderer->ReleaseDrawBuffer(*DeferredUpdateContext.WindowDrawBuffer);
 			}
-		);
+		});
+
+		check(DeferredUpdateContexts.IsEmpty());
+	}
 
-		DeferredUpdateContexts.Empty();
+	if (bScreenshotProcessed)
+	{
+		FlushRenderingCommands();
+		ScreenshotState = {};
 	}
 
-	if (FastPathRenderingDataCleanupList)
+	for (const FWindowToRender& WindowToRender : WindowsToRender)
 	{
-		FastPathRenderingDataCleanupList->Cleanup();
-		FastPathRenderingDataCleanupList = nullptr;
+		SlateWindowRendered.Broadcast(*WindowToRender.Window, &WindowToRender.ViewportInfo->ViewportRHI);
 	}
 
-	// flush the cache if needed
+	FlushPendingDeletes();
 	FontCache->ConditionalFlushCache();
 	ResourceManager->ConditionalFlushAtlases();
 }
 
-
 FIntPoint FSlateRHIRenderer::GenerateDynamicImageResource(const FName InTextureName)
 {
 	check(IsInGameThread());
@@ -2105,9 +1398,6 @@ void FSlateRHIRenderer::RemoveDynamicBrushResource( TSharedPtr<FSlateDynamicImag
 	}
 }
 
-/**
-* Gives the renderer a chance to wait for any render commands to be completed before returning/
-*/
 void FSlateRHIRenderer::FlushCommands() const
 {
 	if (IsInGameThread() || IsInSlateThread())
@@ -2116,10 +1406,6 @@ void FSlateRHIRenderer::FlushCommands() const
 	}
 }
 
-/**
-* Gives the renderer a chance to synchronize with another thread in the event that the renderer runs
-* in a multi-threaded environment.  This function does not return until the sync is complete
-*/
 void FSlateRHIRenderer::Sync() const
 {
 	// Sync game and render thread. Either total sync or allowing one frame lag.
@@ -2128,31 +1414,21 @@ void FSlateRHIRenderer::Sync() const
 	FrameEndSync.Sync(CVarAllowOneFrameThreadLag->GetValueOnAnyThread() != 0);
 }
 
-/**
- * Inline issues a BeginFrame to the RHI.
- * This is to handle cases like Modal dialogs in the UI. The game loop stops while
- * the dialog is open but continues to issue draws. The RHI thinks there are all part of one super long
- * frame until the Modal window is closed.
- */
 void FSlateRHIRenderer::BeginFrame() const
 {
-	ENQUEUE_RENDER_COMMAND(SlateRHIBeginFrame)(
-	   [](FRHICommandListImmediate&)
-	   {
-		   // Suspend stat gathering when running modal dialog 'fake' frame loops
-		   GPU_STATS_SUSPENDFRAME();
-	   }
-	);
+	ENQUEUE_RENDER_COMMAND(SlateRHIBeginFrame)([](FRHICommandListImmediate& RHICmdList)
+	{
+		// Suspend stat gathering when running modal dialog 'fake' frame loops
+		GPU_STATS_SUSPENDFRAME();
+	});
 }
 
 void FSlateRHIRenderer::EndFrame() const
 {
-	ENQUEUE_RENDER_COMMAND(SlateRHIEndFrame)(
-	   [](FRHICommandListImmediate& RHICmdList)
-	   {
-		   RHICmdList.EndFrame();
-	   }
-	);
+	ENQUEUE_RENDER_COMMAND(SlateRHIEndFrame)([](FRHICommandListImmediate& RHICmdList)
+	{
+		RHICmdList.EndFrame();
+	});
 }
 
 void FSlateRHIRenderer::ReloadTextureResources()
@@ -2186,29 +1462,23 @@ void* FSlateRHIRenderer::GetViewportResource(const SWindow& Window)
 {
 	checkSlow(IsThreadSafeForSlateRendering());
 
-	FViewportInfo** InfoPtr = WindowToViewportInfo.Find(&Window);
+	FSlateViewportInfo** InfoPtr = WindowToViewportInfo.Find(&Window);
 
 	if (InfoPtr)
 	{
-		FViewportInfo* ViewportInfo = *InfoPtr;
+		FSlateViewportInfo* ViewportInfo = *InfoPtr;
 
-		// Create the viewport if it doesnt exist
 		if (!IsValidRef(ViewportInfo->ViewportRHI))
 		{
-			// Sanity check dimensions
-			checkf(ViewportInfo->Width <= MAX_VIEWPORT_SIZE && ViewportInfo->Height <= MAX_VIEWPORT_SIZE, TEXT("Invalid window with Width=%u and Height=%u"), ViewportInfo->Width, ViewportInfo->Height);
+			checkf(ViewportInfo->Extent.X <= MAX_VIEWPORT_SIZE && ViewportInfo->Extent.Y <= MAX_VIEWPORT_SIZE, TEXT("Invalid window with Width=%u and Height=%u"), ViewportInfo->Extent.X, ViewportInfo->Extent.Y);
 
-			const bool bFullscreen = IsViewportFullscreen(Window);
-
-			ViewportInfo->ViewportRHI = RHICreateViewport(ViewportInfo->OSWindow, ViewportInfo->Width, ViewportInfo->Height, bFullscreen, ViewportInfo->PixelFormat);
+			ViewportInfo->ViewportRHI = RHICreateViewport(ViewportInfo->OSWindow, ViewportInfo->Extent.X, ViewportInfo->Extent.Y, IsViewportFullscreen(Window), ViewportInfo->PixelFormat);
 		}
 
 		return &ViewportInfo->ViewportRHI;
 	}
-	else
-	{
-		return NULL;
-	}
+
+	return nullptr;
 }
 
 void FSlateRHIRenderer::SetColorVisionDeficiencyType(EColorVisionDeficiency Type, int32 Severity, bool bCorrectDeficiency, bool bShowCorrectionWithDeficiency)
@@ -2223,14 +1493,7 @@ FSlateUpdatableTexture* FSlateRHIRenderer::CreateUpdatableTexture(uint32 Width,
 {
 	const bool bCreateEmptyTexture = true;
 	FSlateTexture2DRHIRef* NewTexture = new FSlateTexture2DRHIRef(Width, Height, GetSlateRecommendedColorFormat(), nullptr, TexCreate_Dynamic, bCreateEmptyTexture);
-	if (IsInRenderingThread())
-	{
-		NewTexture->InitResource(FRHICommandListImmediate::Get());
-	}
-	else
-	{
-		BeginInitResource(NewTexture);
-	}
+	BeginInitResource(NewTexture);
 	return NewTexture;
 }
 
@@ -2262,32 +1525,24 @@ ISlateAtlasProvider* FSlateRHIRenderer::GetTextureAtlasProvider()
 	return nullptr;
 }
 
-
-
 int32 FSlateRHIRenderer::RegisterCurrentScene(FSceneInterface* Scene)
 {
 	check(IsInGameThread());
 	if (Scene && Scene->GetWorld())
 	{
-		// We only want one scene view per world, (todo per player for split screen)
 		CurrentSceneIndex = ActiveScenes.IndexOfByPredicate([&Scene](const FSceneInterface* TestScene) { return TestScene->GetWorld() == Scene->GetWorld(); });
+
 		if (CurrentSceneIndex == INDEX_NONE)
 		{
 			CurrentSceneIndex = ActiveScenes.Add(Scene);
 
-			// We need to keep the ActiveScenes array synchronized with the Policy's ActiveScenes array on
-			// the render thread.
-			FSlateRHIRenderingPolicy* InRenderPolicy = RenderingPolicy.Get();
-			int32 LocalCurrentSceneIndex = CurrentSceneIndex;
-			ENQUEUE_RENDER_COMMAND(RegisterCurrentSceneOnPolicy)(
-				[InRenderPolicy, Scene, LocalCurrentSceneIndex](FRHICommandListImmediate& RHICmdList)
+			if (CurrentSceneIndex >= 0)
 			{
-				if (LocalCurrentSceneIndex != -1)
+				ENQUEUE_RENDER_COMMAND(RegisterCurrentSceneOnPolicy)([RenderingPolicy = RenderingPolicy.Get(), Scene, CurrentSceneIndex = CurrentSceneIndex](FRHICommandListBase&)
 				{
-					InRenderPolicy->AddSceneAt(Scene, LocalCurrentSceneIndex);
-				}
+					RenderingPolicy->AddSceneAt(Scene, CurrentSceneIndex);
+				});
 			}
-			);
 		}
 	}
 	else
@@ -2315,15 +1570,10 @@ void FSlateRHIRenderer::ClearScenes()
 		CurrentSceneIndex = -1;
 		ActiveScenes.Empty();
 
-		// We need to keep the ActiveScenes array synchronized with the Policy's ActiveScenes array on
-		// the render thread.
-		FSlateRenderingPolicy* InRenderPolicy = RenderingPolicy.Get();
-		ENQUEUE_RENDER_COMMAND(ClearScenesOnPolicy)(
-			[InRenderPolicy](FRHICommandListImmediate& RHICmdList)
+		ENQUEUE_RENDER_COMMAND(ClearScenesOnPolicy)([RenderingPolicy = RenderingPolicy.Get()](FRHICommandListBase&)
 		{
-			InRenderPolicy->ClearScenes();
-		}
-		);
+			RenderingPolicy->ClearScenes();
+		});
 	}
 }
 
@@ -2332,106 +1582,59 @@ EPixelFormat FSlateRHIRenderer::GetSlateRecommendedColorFormat()
 	return bIsStandaloneStereoOnlyDevice ? PF_R8G8B8A8 : PF_B8G8R8A8;
 }
 
-FRHICOMMAND_MACRO(FClearCachedRenderingDataCommand)
-{
-public:
-	FClearCachedRenderingDataCommand(FSlateCachedFastPathRenderingData* InCachedRenderingData)
-		: CachedRenderingData(InCachedRenderingData)
-	{
-		
-	}
-
-	void Execute(FRHICommandListBase& CmdList)
-	{
-		delete CachedRenderingData;
-	}
-
-private:
-	FSlateCachedFastPathRenderingData* CachedRenderingData;
-};
-
-FRHICOMMAND_MACRO(FClearCachedElementDataCommand)
-{
-public:
-	FClearCachedElementDataCommand(FSlateCachedElementData* InCachedElementData)
-		: CachedElementData(InCachedElementData)
-	{
-
-	}
-
-	void Execute(FRHICommandListBase& CmdList)
-	{
-		delete CachedElementData;
-	}
-
-private:
-	FSlateCachedElementData* CachedElementData;
-};
-
 void FSlateRHIRenderer::DestroyCachedFastPathRenderingData(FSlateCachedFastPathRenderingData* CachedRenderingData)
 {
 	check(CachedRenderingData);
-
-	if (!FastPathRenderingDataCleanupList)
-	{
-		// This will be deleted later on the rendering thread
-		FastPathRenderingDataCleanupList = new FFastPathRenderingDataCleanupList;
-	}
-
-	FastPathRenderingDataCleanupList->FastPathRenderingDataToRemove.Add(CachedRenderingData);
+	PendingDeletes.CachedRenderingData.Emplace(CachedRenderingData);
 }
 
 void FSlateRHIRenderer::DestroyCachedFastPathElementData(FSlateCachedElementData* CachedElementData)
 {
 	check(CachedElementData);
-
-	// Cached data should be destroyed in a thread safe way.  If there is an rhi thread it could be reading from the data to copy it into a vertex buffer
-	// so delete it on the rhi thread if necessary, otherwise delete it on the render thread
-	ENQUEUE_RENDER_COMMAND(ClearCachedElementData)(
-		[CachedElementData](FRHICommandListImmediate& RHICmdList)
-	{
-		if (!RHICmdList.Bypass())
-		{
-			new (RHICmdList.AllocCommand<FClearCachedElementDataCommand>()) FClearCachedElementDataCommand(CachedElementData);
-		}
-		else
-		{
-			FClearCachedElementDataCommand Cmd(CachedElementData);
-			Cmd.Execute(RHICmdList);
-		}
-	});
+	PendingDeletes.CachedElementData.Emplace(CachedElementData);
 }
 
-#if WITH_EDITORONLY_DATA
-namespace SlateRendererUtil
+void FSlateRHIRenderer::FlushPendingDeletes()
 {
-	bool bSlateShadersInitialized = false;
-	FDelegateHandle GlobalShaderCompilationDelegateHandle;
-
-	bool AreShadersInitialized()
+	if (!PendingDeletes.IsEmpty())
 	{
-		if (!bSlateShadersInitialized)
+		ENQUEUE_RENDER_COMMAND(SlateDeferredDelete)([PendingDeletes = MoveTemp(PendingDeletes)](FRHICommandListBase&)
 		{
-			bSlateShadersInitialized = IsGlobalShaderMapComplete(TEXT("SlateElement"));
-			// if shaders are initialized, cache the value until global shaders gets recompiled.
-			if (bSlateShadersInitialized)
+			for (FSlateCachedFastPathRenderingData* Data : PendingDeletes.CachedRenderingData)
 			{
-				GlobalShaderCompilationDelegateHandle = GetOnGlobalShaderCompilation().AddLambda([]()
-				{
-					bSlateShadersInitialized = false;
-					GetOnGlobalShaderCompilation().Remove(GlobalShaderCompilationDelegateHandle);
-				});
+				delete Data;
 			}
-		}
-		return bSlateShadersInitialized;
+
+			for (FSlateCachedElementData* Data : PendingDeletes.CachedElementData)
+			{
+				delete Data;
+			}
+		});
+		PendingDeletes = {};
 	}
 }
-#endif
 
 bool FSlateRHIRenderer::AreShadersInitialized() const
 {
 #if WITH_EDITORONLY_DATA
-	return SlateRendererUtil::AreShadersInitialized();
+	static bool bSlateShadersInitialized = false;
+	static FDelegateHandle GlobalShaderCompilationDelegateHandle;
+
+	if (!bSlateShadersInitialized)
+	{
+		bSlateShadersInitialized = IsGlobalShaderMapComplete(TEXT("SlateElement"));
+
+		// if shaders are initialized, cache the value until global shaders gets recompiled.
+		if (bSlateShadersInitialized)
+		{
+			GlobalShaderCompilationDelegateHandle = GetOnGlobalShaderCompilation().AddLambda([]()
+			{
+				bSlateShadersInitialized = false;
+				GetOnGlobalShaderCompilation().Remove(GlobalShaderCompilationDelegateHandle);
+			});
+		}
+	}
+	return bSlateShadersInitialized;
 #else
 	return true;
 #endif
@@ -2439,9 +1642,9 @@ bool FSlateRHIRenderer::AreShadersInitialized() const
 
 void FSlateRHIRenderer::InvalidateAllViewports()
 {
-	for (TMap< const SWindow*, FViewportInfo*>::TIterator It(WindowToViewportInfo); It; ++It)
+	for (auto& Entry : WindowToViewportInfo)
 	{
-		It.Value()->ViewportRHI = nullptr;
+		Entry.Value->ViewportRHI = nullptr;
 	}
 }
 
@@ -2459,12 +1662,6 @@ void FSlateRHIRenderer::ReleaseAccessedResources(bool bImmediatelyFlush)
 	{
 		// Increment resource version to allow buffers to shrink or cached structures to clean up.
 		ResourceVersion++;
-
-		// Release resources generated specifically by the rendering policy if we are flushing.
-		// This should NOT be done unless flushing
-		RenderingPolicy->FlushGeneratedResources();
-
-		//FlushCommands();
 	}
 }
 
@@ -2472,30 +1669,25 @@ void FSlateRHIRenderer::RequestResize(const TSharedPtr<SWindow>& Window, uint32
 {
 	checkSlow(IsThreadSafeForSlateRendering());
 
-	FViewportInfo* ViewInfo = WindowToViewportInfo.FindRef(Window.Get());
+	FSlateViewportInfo* ViewInfo = WindowToViewportInfo.FindRef(Window.Get());
 
 	if (ViewInfo)
 	{
-		ViewInfo->DesiredWidth = NewWidth;
-		ViewInfo->DesiredHeight = NewHeight;
+		ViewInfo->ExtentToResizeTo.X = NewWidth;
+		ViewInfo->ExtentToResizeTo.Y = NewHeight;
 	}
 }
 
 void FSlateRHIRenderer::SetWindowRenderTarget(const SWindow& Window, IViewportRenderTargetProvider* Provider)
 {
-	FViewportInfo* ViewInfo = WindowToViewportInfo.FindRef(&Window);
+	FSlateViewportInfo* ViewInfo = WindowToViewportInfo.FindRef(&Window);
 	if (ViewInfo)
 	{
-		ViewInfo->RTProvider = Provider;
+		ViewInfo->RenderTargetProvider = Provider;
 	}
 }
 
-FSlateRHIRenderingPolicyInterface FSlateRHIRenderer::GetRenderingPolicyInterface()
-{
-	return FSlateRHIRenderingPolicyInterface(RenderingPolicy.Get());
-}
-
-void FSlateRHIRenderer::AddWidgetRendererUpdate(const struct FRenderThreadUpdateContext& Context, bool bDeferredRenderTargetUpdate)
+void FSlateRHIRenderer::AddWidgetRendererUpdate(const FRenderThreadUpdateContext& Context, bool bDeferredRenderTargetUpdate)
 {
 	if (bDeferredRenderTargetUpdate)
 	{
@@ -2503,104 +1695,11 @@ void FSlateRHIRenderer::AddWidgetRendererUpdate(const struct FRenderThreadUpdate
 	}
 	else
 	{
-		// Enqueue a command to unlock the draw buffer after all windows have been drawn
-		FRenderThreadUpdateContext InContext = Context;
-		ENQUEUE_RENDER_COMMAND(DrawWidgetRendererImmediate)(
-			[InContext](FRHICommandListImmediate& RHICmdList)
-			{
-				InContext.Renderer->DrawWindowToTarget_RenderThread(RHICmdList, InContext);
-			});
-	}
-}
-
-FSlateEndDrawingWindowsCommand::FSlateEndDrawingWindowsCommand(FSlateRHIRenderingPolicy& InPolicy, FSlateDrawBuffer* InDrawBuffer)
-	: Policy(InPolicy)
-	, DrawBuffer(InDrawBuffer)
-{}
-
-void FSlateEndDrawingWindowsCommand::Execute(FRHICommandListBase& CmdList)
-{
-	Policy.EndDrawingWindows();
-}
-
-void FSlateEndDrawingWindowsCommand::EndDrawingWindows(FRHICommandListImmediate& RHICmdList, FSlateDrawBuffer* DrawBuffer, FSlateRHIRenderingPolicy& Policy)
-{
-	if (!RHICmdList.Bypass())
-	{
-		ALLOC_COMMAND_CL(RHICmdList, FSlateEndDrawingWindowsCommand)(Policy, DrawBuffer);
-	}
-	else
-	{
-		FSlateEndDrawingWindowsCommand Cmd(Policy, DrawBuffer);
-		Cmd.Execute(RHICmdList);
-	}
-}
-
-FSlateReleaseDrawBufferCommand::FSlateReleaseDrawBufferCommand(FSlateDrawBuffer* InDrawBuffer)
-	: DrawBuffer(InDrawBuffer)
-{}
-
-void FSlateReleaseDrawBufferCommand::Execute(FRHICommandListBase& CmdList)
-{
-	DrawBuffer->Unlock();
-}
-
-void FSlateReleaseDrawBufferCommand::ReleaseDrawBuffer(FRHICommandListImmediate& RHICmdList, FSlateDrawBuffer* DrawBuffer)
-{
-	if (!RHICmdList.Bypass())
-	{
-		ALLOC_COMMAND_CL(RHICmdList, FSlateReleaseDrawBufferCommand)(DrawBuffer);
-	}
-	else
-	{
-		FSlateReleaseDrawBufferCommand Cmd(DrawBuffer);
-		Cmd.Execute(RHICmdList);
-	}
-}
-
-
-struct FClearCachedRenderingDataCommand2 final : public FRHICommand < FClearCachedRenderingDataCommand2 >
-{
-public:
-	FClearCachedRenderingDataCommand2(FFastPathRenderingDataCleanupList* InCleanupList)
-		: CleanupList(InCleanupList)
-	{
-
-	}
-
-	void Execute(FRHICommandListBase& CmdList)
-	{
-		delete CleanupList;
-	}
-
-private:
-	FFastPathRenderingDataCleanupList* CleanupList;
-};
-
-
-FFastPathRenderingDataCleanupList::~FFastPathRenderingDataCleanupList()
-{
-	for (FSlateCachedFastPathRenderingData* Data : FastPathRenderingDataToRemove)
-	{
-		delete Data;
-	}
-}
-
-void FFastPathRenderingDataCleanupList::Cleanup()
-{
-	// Cached data should be destroyed in a thread safe way.  If there is an rhi thread it could be reading from the data to copy it into a vertex buffer
-	// so delete it on the rhi thread if necessary, otherwise delete it on the render thread
-	ENQUEUE_RENDER_COMMAND(ClearCachedRenderingData)(
-		[CleanupList = this](FRHICommandListImmediate& RHICmdList)
-	{
-		if (!RHICmdList.Bypass())
+		ENQUEUE_RENDER_COMMAND(DrawWidgetRendererImmediate)([Context](FRHICommandListImmediate& RHICmdList)
 		{
-			new (RHICmdList.AllocCommand<FClearCachedRenderingDataCommand2>()) FClearCachedRenderingDataCommand2(CleanupList);
-		}
-		else
-		{
-			FClearCachedRenderingDataCommand2 Cmd(CleanupList);
-			Cmd.Execute(RHICmdList);
-		}
-	});
-}
+			FRDGBuilder GraphBuilder(RHICmdList);
+			Context.Renderer->DrawWindowToTarget_RenderThread(GraphBuilder, Context);
+			GraphBuilder.Execute();
+		});
+	}
+}
\ No newline at end of file
