Clean-up of Slate RHI renderer and conversion to RDG.

diff --git a/Engine/Source/Runtime/SlateRHIRenderer/Private/SlateRHIRenderer.h b/Engine/Source/Runtime/SlateRHIRenderer/Private/SlateRHIRenderer.h
--- a/Engine/Source/Runtime/SlateRHIRenderer/Private/SlateRHIRenderer.h
+++ b/Engine/Source/Runtime/SlateRHIRenderer/Private/SlateRHIRenderer.h
@@ -18,16 +18,14 @@
 class FSlateElementBatcher;
 class FSlateRHIRenderingPolicy;
 class FSlateRHIRenderingPolicyInterface;
+class FSlateRHIPostBufferProcessorProxy;
 class USlateRHIPostBufferProcessor;
 class USlateRHIRendererSettings;
 class ISlateStyle;
 class SWindow;
 struct FSlatePostSettings;
-struct Rect;
 
-template<typename TCmd, typename NameType> struct FRHICommand;
-
-typedef TArray<FRenderThreadUpdateContext, TInlineAllocator<2>> FDeferredUpdateContextList;
+extern FMatrix CreateSlateProjectionMatrix(uint32 Width, uint32 Height);
 
 // Number of draw buffers that can be active at any given time
 const uint32 NumDrawBuffers = 3;
@@ -35,151 +33,17 @@ const uint32 NumDrawBuffers = 3;
 // Enable to visualize overdraw in Slate
 #define WITH_SLATE_VISUALIZERS !(UE_BUILD_SHIPPING || UE_BUILD_TEST)
 
-#if WITH_SLATE_VISUALIZERS
-extern TAutoConsoleVariable<int32> CVarShowSlateOverdraw;
-extern TAutoConsoleVariable<int32> CVarShowSlateBatching;
-#endif
-
-class FSlateBackBuffer : public FRenderTarget
-{
-public:
-	FSlateBackBuffer(FTextureRHIRef InRenderTargetTexture, FIntPoint InSizeXY)
-		: SizeXY(InSizeXY)
-	{
-		RenderTargetTextureRHI = InRenderTargetTexture;
-	}
-	virtual FIntPoint GetSizeXY() const override { return SizeXY; }
-private:
-	FIntPoint SizeXY;
-};
-
+struct FFastPathRenderingDataCleanUpList;
 
-/** An RHI Representation of a viewport with cached width and height for detecting resizes */
-struct FViewportInfo : public FRenderResource
-{
-	/** The projection matrix used in the viewport */
-	FMatrix ProjectionMatrix;	
-	/** The viewport rendering handle */
-	FViewportRHIRef ViewportRHI;
-	/** The depth buffer texture if any */
-	FTextureRHIRef DepthStencil;
-
-	// Buffers used in HDR compositing
-	/** sRGB UI render target */
-	TRefCountPtr<IPooledRenderTarget> UITargetRT;
-	TRefCountPtr<IPooledRenderTarget> UITargetRTMask;
-
-	/** Color-space LUT for HDR UI composition. */
-	FTextureRHIRef ColorSpaceLUT;
-	int32 ColorSpaceLUTOutputDevice;
-	int32 ColorSpaceLUTOutputGamut;
-		
-	//FTextureRHIRef RenderTargetTexture;
-	/** The OS Window handle (for recreating the viewport) */
-	void* OSWindow;
-	/** The actual width of the viewport */
-	uint32 Width;
-	/** The actual height of the viewport */
-	uint32 Height;
-	/** The desired width of the viewport */
-	uint32 DesiredWidth;
-	/** The desired height of the viewport */
-	uint32 DesiredHeight;
-	/** Whether or not the viewport requires a stencil test */
-	bool bRequiresStencilTest;
-	/** Whether or not the viewport is in fullscreen */
-	bool bFullscreen;
-	/** The desired pixel format for this viewport */
-	EPixelFormat PixelFormat;
-	/** The desired SDR pixel format for this viewport */
-	EPixelFormat SDRPixelFormat;
-	/** Color gamut for output to HDR display */
-	EDisplayColorGamut HDRDisplayColorGamut;
-	/** Device format for output to HDR display */
-	EDisplayOutputFormat HDRDisplayOutputFormat;
-
-	IViewportRenderTargetProvider* RTProvider;
-	
-	/** Whether is in a HDR Color Space */
-	bool bHDREnabled;
-
-	/** Whether the scene output is in HDR. bHDREnabled only affects Slate rendering */
-	bool bSceneHDREnabled;
-
-	/** FRenderResource interface */
-	virtual void InitRHI(FRHICommandListBase& RHICmdList) override;
-	virtual void ReleaseRHI() override;
-	virtual void ReleaseResource() override;
-
-	FViewportInfo()
-		:	ColorSpaceLUTOutputDevice(0),
-			ColorSpaceLUTOutputGamut(0),
-			OSWindow(NULL), 
-			Width(0),
-			Height(0),
-			DesiredWidth(0),
-			DesiredHeight(0),
-			bRequiresStencilTest(false),
-			bFullscreen(false),
-			PixelFormat(EPixelFormat::PF_Unknown),
-			SDRPixelFormat(EPixelFormat::PF_Unknown),
-			HDRDisplayColorGamut(EDisplayColorGamut::sRGB_D65),
-			HDRDisplayOutputFormat(EDisplayOutputFormat::SDR_sRGB),
-			RTProvider(nullptr),
-			bHDREnabled(false),
-			bSceneHDREnabled(false)
-	{
-	}
-
-	~FViewportInfo()
-	{
-		DepthStencil.SafeRelease();
-		ColorSpaceLUT.SafeRelease();
-	}
-
-	void ConditionallyUpdateDepthBuffer(bool bInRequiresStencilTest, uint32 Width, uint32 Height);
-	void RecreateDepthBuffer_RenderThread();
-
-	FTextureRHIRef GetRenderTargetTexture() const
-	{
-		if (RTProvider)
-		{
-			FSlateShaderResource* RenderTargetTexture = RTProvider->GetViewportRenderTargetTexture();
-			if( RenderTargetTexture )
-			{
-				FSlateRenderTargetRHI* RHITarget = (FSlateRenderTargetRHI*)RenderTargetTexture;
-				return RHITarget->GetTypedResource();
-			}
-		}
-		return nullptr;
-	}	
-};
-
-struct FFastPathRenderingDataCleanupList
-{
-	TArray<FSlateCachedFastPathRenderingData*, TInlineAllocator<20>> FastPathRenderingDataToRemove;
-
-	~FFastPathRenderingDataCleanupList();
-	void Cleanup();
-};
+struct FSlateDrawWindowPassInputs;
+struct FSlateViewportInfo;
+struct FSlatePostProcessUpdateRequest;
 
 /** A Slate rendering implementation for Unreal engine */
-class FSlateRHIRenderer : public FSlateRenderer
+class FSlateRHIRenderer final : public FSlateRenderer
 {
 public:
 	FSlateRHIRenderer( TSharedRef<FSlateFontServices> InSlateFontServices, TSharedRef<FSlateRHIResourceManager> InResourceManager );
-	~FSlateRHIRenderer();
-
-	/**
-	 * Creates a projection matrix for use when rendering an SWindow
-	 *
-	 * @param Width 	The width of the window
-	 * @param Height	The height of the window
-	 * @return The created projection matrix
-	 */
-	static FMatrix CreateProjectionMatrix( uint32 Width, uint32 Height );
-	static int32 GetDrawToVRRenderTarget();
-	static int32 GetProcessSlatePostBuffers();
 
 	/** FSlateRenderer interface */
 	virtual bool Initialize() override;
@@ -222,9 +86,6 @@ public:
 	virtual void EndFrame() const override;
 	virtual void AddWidgetRendererUpdate(const struct FRenderThreadUpdateContext& Context, bool bDeferredRenderTargetUpdate) override;
 
-	/** Draws windows from a FSlateDrawBuffer on the render thread */
-	void DrawWindow_RenderThread(FRHICommandListImmediate& RHICmdList, FViewportInfo& ViewportInfo, FSlateWindowElementList& WindowElementList, const struct FSlateDrawWindowCommandParams& DrawCommandParams);
-
 	/**
 	 * Reloads texture resources from disk                   
 	 */
@@ -247,28 +108,22 @@ public:
 
 	virtual void SetWindowRenderTarget(const SWindow& Window, class IViewportRenderTargetProvider* Provider) override;
 
-	FSlateRHIRenderingPolicyInterface GetRenderingPolicyInterface();
-
 private:
+	void DrawWindow_RenderThread(FRDGBuilder& GraphBuilder, const FSlateDrawWindowPassInputs& Inputs);
+
 	/** Loads all known textures from Slate styles */
 	void LoadUsedTextures();
 
-	/**
-	 * Resizes the viewport for a window if needed
-	 * 
-	 * @param ViewportInfo	The viewport to resize
-	 * @param Width			The width that we should size to
-	 * @param Height		The height that we shoudl size to
-	 * @param bFullscreen	If we should be in fullscreen
-	 */
-	void ConditionalResizeViewport( FViewportInfo* ViewportInfo, uint32 Width, uint32 Height, bool bFullscreen, SWindow* Window);
-	
+	void ResizeViewportIfNeeded(FSlateViewportInfo* ViewportInfo, FIntPoint ExtentToResizeTo, bool bFullscreen, SWindow* Window);
+
+	EPixelFormat GetViewportPixelFormat(const SWindow& Window, bool bDisplayFormatIsHDR);
+
 	/** 
 	 * Creates necessary resources to render a window and sends draw commands to the rendering thread
 	 *
 	 * @param WindowDrawBuffer	The buffer containing elements to draw 
 	 */
-	void DrawWindows_Private( FSlateDrawBuffer& InWindowDrawBuffer );
+	void DrawWindows_Private(FSlateDrawBuffer& InWindowDrawBuffer);
 
 	/**
 	 * Delete the updateable textures we've marked for delete that have already had their GPU resources released, but may
@@ -278,21 +133,17 @@ private:
 
 	virtual void OnVirtualDesktopSizeChanged(const FDisplayMetrics& NewDisplayMetric);
 
-private:
 	/** A mapping of SWindows to their RHI implementation */
-	TMap< const SWindow*, FViewportInfo*> WindowToViewportInfo;
-
-	/** View matrix used by all windows */
-	FMatrix ViewMatrix;
+	TMap<const SWindow*, FSlateViewportInfo*> WindowToViewportInfo;
 
 	/** Keep a pointer around for when we have deferred drawing happening */
-	FSlateDrawBuffer* EnqueuedWindowDrawBuffer;
+	FSlateDrawBuffer* EnqueuedWindowDrawBuffer = nullptr;
 
 	/** Double buffered draw buffers so that the rendering thread can be rendering windows while the game thread is setting up for next frame */
 	FSlateDrawBuffer DrawBuffers[NumDrawBuffers];
 
 	/** The draw buffer which is currently free for use by the game thread */
-	uint8 FreeBufferIndex;
+	uint8 FreeBufferIndex = 0;
 
 	/** Element batcher which renders draw elements */
 	TUniquePtr<FSlateElementBatcher> ElementBatcher;
@@ -305,56 +156,44 @@ private:
 
 	TArray<TSharedPtr<FSlateDynamicImageBrush>> DynamicBrushesToRemove[NumDrawBuffers];
 
-	FFastPathRenderingDataCleanupList* FastPathRenderingDataCleanupList;
-
-	FDeferredUpdateContextList DeferredUpdateContexts;
-
-	bool bIsStandaloneStereoOnlyDevice;
-	bool bTakingAScreenShot;
-	bool bUpdateHDRDisplayInformation;
-	ESlatePostRT bShrinkPostBufferRequested;
-	uint64 LastFramesPostBufferUsed[(uint8)ESlatePostRT::Num];
-	FRenderCommandFence SlatePostRTFences[(uint8)ESlatePostRT::Num];
-	FIntRect ScreenshotRect;
-	FViewportInfo* ScreenshotViewportInfo;
-	TArray<FColor>* OutScreenshotData;
-	TArray<FLinearColor>* OutHDRScreenshotData;
+	struct
+	{
+		TArray<FSlateCachedFastPathRenderingData*, FConcurrentLinearArrayAllocator> CachedRenderingData;
+		TArray<FSlateCachedElementData*, FConcurrentLinearArrayAllocator> CachedElementData;
 
-	/** These are state management variables for Scenes on the game thread. A similar copy exists on the RHI Rendering Policy for the rendering thread.*/
-	TArray<FSceneInterface*, TInlineAllocator<4>> ActiveScenes;
-	int32 CurrentSceneIndex;
-	
-	/** Version that increments when it is okay to clean up older cached resources */
-	uint32 ResourceVersion;
-};
+		bool IsEmpty() const { return CachedRenderingData.IsEmpty() && CachedElementData.IsEmpty(); }
 
-struct FSlateEndDrawingWindowsCommandString
-{
-	static const TCHAR* TStr() { return TEXT("FSlateEndDrawingWindowsCommand"); }
-};
-struct FSlateEndDrawingWindowsCommand final : public FRHICommand < FSlateEndDrawingWindowsCommand, FSlateEndDrawingWindowsCommandString >
-{
-	FSlateRHIRenderingPolicy& Policy;
-	FSlateDrawBuffer* DrawBuffer;
+	} PendingDeletes;
 
-	FSlateEndDrawingWindowsCommand(FSlateRHIRenderingPolicy& InPolicy, FSlateDrawBuffer* InDrawBuffer);
+	void FlushPendingDeletes();
 
-	void Execute(FRHICommandListBase& CmdList);
+	TArray<FRenderThreadUpdateContext, FConcurrentLinearArrayAllocator> DeferredUpdateContexts;
 
-	static void EndDrawingWindows(FRHICommandListImmediate& RHICmdList, FSlateDrawBuffer* DrawBuffer, FSlateRHIRenderingPolicy& Policy);
-};
+	bool bIsStandaloneStereoOnlyDevice = false;
+	bool bUpdateHDRDisplayInformation = false;
+	ESlatePostRT bShrinkPostBufferRequested = ESlatePostRT::None;
+	uint64 LastFramesPostBufferUsed[(uint8)ESlatePostRT::Num];
+	FRenderCommandFence SlatePostRTFences[(uint8)ESlatePostRT::Num];
 
-struct FSlateReleaseDrawBufferCommandString
-{
-	static const TCHAR* TStr() { return TEXT("FSlateReleaseDrawBufferCommand"); }
-};
-struct FSlateReleaseDrawBufferCommand final : public FRHICommand < FSlateReleaseDrawBufferCommand, FSlateReleaseDrawBufferCommandString >
-{
-	FSlateDrawBuffer* DrawBuffer;
+	struct
+	{
+		TStaticArray<uint64, (int32)ESlatePostRT::Num> LastUsedFrameCounter{ InPlace, 0 };
+	
+	} PostProcessRenderTargets;
 
-	FSlateReleaseDrawBufferCommand(FSlateDrawBuffer* InDrawBuffer);
+	struct
+	{
+		FIntRect ViewRect;
+		FSlateViewportInfo* ViewportToCapture = nullptr;
+		TArray<FColor>* ColorData = nullptr;
+		TArray<FLinearColor>* ColorDataHDR = nullptr;
 
-	void Execute(FRHICommandListBase& CmdList);
+	} ScreenshotState;
 
-	static void ReleaseDrawBuffer(FRHICommandListImmediate& RHICmdList, FSlateDrawBuffer* DrawBuffer);
+	/** These are state management variables for Scenes on the game thread. A similar copy exists on the RHI Rendering Policy for the rendering thread.*/
+	TArray<FSceneInterface*, TInlineAllocator<4>> ActiveScenes;
+	int32 CurrentSceneIndex = -1;
+	
+	/** Version that increments when it is okay to clean up older cached resources */
+	uint32 ResourceVersion = 0;
 };
