Clean-up of Slate RHI renderer and conversion to RDG.

diff --git a/Engine/Source/Runtime/SlateRHIRenderer/Private/SlateRHIRenderingPolicy.cpp b/Engine/Source/Runtime/SlateRHIRenderer/Private/SlateRHIRenderingPolicy.cpp
--- a/Engine/Source/Runtime/SlateRHIRenderer/Private/SlateRHIRenderingPolicy.cpp
+++ b/Engine/Source/Runtime/SlateRHIRenderer/Private/SlateRHIRenderingPolicy.cpp
@@ -29,8 +29,10 @@
 #include "DeviceProfiles/DeviceProfile.h"
 #include "DeviceProfiles/DeviceProfileManager.h"
 #include "Types/SlateConstants.h"
-#include "RenderGraphResources.h"
+#include "RenderGraphUtils.h"
 #include "SceneRenderTargetParameters.h"
+#include "SceneInterface.h"
+#include "Containers/StaticBitArray.h"
 
 extern void UpdateNoiseTextureParameters(FViewUniformShaderParameters& ViewUniformShaderParameters);
 
@@ -45,529 +47,86 @@ DECLARE_DWORD_COUNTER_STAT(TEXT("Clips (Scissor)"), STAT_SlateScissorClips, STAT
 DECLARE_DWORD_COUNTER_STAT(TEXT("Clips (Stencil)"), STAT_SlateStencilClips, STATGROUP_Slate);
 
 #if WITH_SLATE_DEBUGGING
-int32 SlateEnableDrawEvents = 1;
+static int32 GSlateEnableDrawEvents = 1;
 #else
-int32 SlateEnableDrawEvents = 0;
-#endif
-static FAutoConsoleVariableRef CVarSlateEnableDrawEvents(TEXT("Slate.EnableDrawEvents"), SlateEnableDrawEvents, TEXT("."), ECVF_Default);
-
-#if WITH_SLATE_DEBUGGING
-int32 BatchToDraw = -1;
-static FAutoConsoleVariableRef CVarSlateDrawBatchNum(TEXT("Slate.DrawBatchNum"), BatchToDraw, TEXT("."), ECVF_Default);
+static int32 GSlateEnableDrawEvents = 0;
 #endif
+static FAutoConsoleVariableRef CVarGSlateEnableDrawEvents(TEXT("Slate.EnableDrawEvents"), GSlateEnableDrawEvents, TEXT("."), ECVF_Default);
 
 #if !(UE_BUILD_SHIPPING || UE_BUILD_TEST)
-	#define SLATE_DRAW_EVENT( RHICmdList, EventName             ) SCOPED_CONDITIONAL_DRAW_EVENT( RHICmdList, EventName, (SlateEnableDrawEvents != 0));
-	#define SLATE_DRAW_EVENTF(RHICmdList, EventName, Format, ...) SCOPED_CONDITIONAL_DRAW_EVENTF(RHICmdList, EventName, (SlateEnableDrawEvents != 0), Format, ##__VA_ARGS__);
+	#define SLATE_DRAW_EVENT( RHICmdList, EventName             ) SCOPED_CONDITIONAL_DRAW_EVENT( RHICmdList, EventName, (GSlateEnableDrawEvents != 0));
+	#define SLATE_DRAW_EVENTF(RHICmdList, EventName, Format, ...) SCOPED_CONDITIONAL_DRAW_EVENTF(RHICmdList, EventName, (GSlateEnableDrawEvents != 0), Format, ##__VA_ARGS__);
 #else
 	#define SLATE_DRAW_EVENT( RHICmdList, EventName             )
 	#define SLATE_DRAW_EVENTF(RHICmdList, EventName, Format, ...)
 #endif
 
-FSlateRHIRenderingPolicy::FSlateRHIRenderingPolicy(TSharedRef<FSlateFontServices> InSlateFontServices, TSharedRef<FSlateRHIResourceManager> InResourceManager, TOptional<int32> InitialBufferSize)
-	: FSlateRenderingPolicy(InSlateFontServices, 0)
-	, PostProcessor(new FSlatePostProcessor)
-	, ResourceManager(InResourceManager)
-	, bGammaCorrect(true)
-	, bApplyColorDeficiencyCorrection(true)
-	, InitialBufferSizeOverride(InitialBufferSize)
-	, LastDeviceProfile(nullptr)
-{
-	InitResources();
-}
-
-void FSlateRHIRenderingPolicy::InitResources()
-{
-	int32 NumVertices = 100;
-
-	if ( InitialBufferSizeOverride.IsSet() )
-	{
-		NumVertices = InitialBufferSizeOverride.GetValue();
-	}
-	else if ( GConfig )
-	{
-		int32 NumVertsInConfig = 0;
-		if ( GConfig->GetInt(TEXT("SlateRenderer"), TEXT("NumPreallocatedVertices"), NumVertsInConfig, GEngineIni) )
-		{
-			NumVertices = NumVertsInConfig;
-		}
-	}
-
-	// Always create a little space but never allow it to get too high
-#if !SLATE_USE_32BIT_INDICES
-	NumVertices = FMath::Clamp(NumVertices, 100, 65535);
-#else
-	NumVertices = FMath::Clamp(NumVertices, 100, 1000000);
+#if WITH_SLATE_VISUALIZERS
+extern TAutoConsoleVariable<int32> CVarShowSlateOverdraw;
+extern TAutoConsoleVariable<int32> CVarShowSlateBatching;
 #endif
 
-	UE_LOG(LogSlate, Verbose, TEXT("Allocating space for %d vertices"), NumVertices);
-
-	SourceVertexBuffer.Init(NumVertices);
-	SourceIndexBuffer.Init(NumVertices);
-
-	BeginInitResource(&StencilVertexBuffer);
-}
-
-void FSlateRHIRenderingPolicy::ReleaseResources()
-{
-	SourceVertexBuffer.Destroy();
-	SourceIndexBuffer.Destroy();
+//////////////////////////////////////////////////////////////////////////
 
-	BeginReleaseResource(&StencilVertexBuffer);
-}
+FSlateRHIRenderingPolicy::FSlateRHIRenderingPolicy(TSharedRef<FSlateFontServices> InSlateFontServices, TSharedRef<FSlateRHIResourceManager> InResourceManager)
+	: FSlateRenderingPolicy(InSlateFontServices, 0)
+	, ResourceManager(InResourceManager)
+{}
 
-void FSlateRHIRenderingPolicy::BeginDrawingWindows()
+void FSlateRHIRenderingPolicy::AddSceneAt(FSceneInterface* Scene, int32 Index)
 {
-	check( IsInRenderingThread() );
+	ResourceManager->AddSceneAt(Scene, Index);
 }
 
-void FSlateRHIRenderingPolicy::EndDrawingWindows()
+void FSlateRHIRenderingPolicy::ClearScenes()
 {
-	check( IsInParallelRenderingThread() );
+	ResourceManager->ClearScenes();
 }
 
-void FSlateRHIRenderingPolicy::BuildRenderingBuffers(FRHICommandListImmediate& RHICmdList, FSlateBatchData& InBatchData)
-{
-	SCOPE_CYCLE_COUNTER(STAT_SlateUpdateBufferRTTime);
-
-	// Should only be called by the rendering thread
-	check(IsInRenderingThread());
-
-	// Merge together batches for less draw calls
-	InBatchData.MergeRenderBatches();
-
-	const FSlateVertexArray& FinalVertexData = InBatchData.GetFinalVertexData();
-	const FSlateIndexArray& FinalIndexData = InBatchData.GetFinalIndexData();
-
-	const int32 NumBatchedVertices = FinalVertexData.Num();
-	const int32 NumBatchedIndices = FinalIndexData.Num();
-
-	if (InBatchData.GetRenderBatches().Num() > 0 && NumBatchedVertices > 0 && NumBatchedIndices > 0)
-	{
-		const bool bShouldShrinkResources = false;
-		SourceVertexBuffer.PreFillBuffer(RHICmdList, NumBatchedVertices, bShouldShrinkResources);
-		SourceIndexBuffer.PreFillBuffer(RHICmdList, NumBatchedIndices, bShouldShrinkResources);
-
-		const uint32 RequiredVertexBufferSize = NumBatchedVertices * sizeof(FSlateVertex);
-		const uint32 RequiredIndexBufferSize = NumBatchedIndices * sizeof(SlateIndex);
-
-		uint8* DstVertexData = (uint8*)RHICmdList.LockBuffer(SourceVertexBuffer.VertexBufferRHI, 0, RequiredVertexBufferSize, RLM_WriteOnly);
-		uint8* DstIndexData  = (uint8*)RHICmdList.LockBuffer(SourceIndexBuffer.IndexBufferRHI, 0, RequiredIndexBufferSize, RLM_WriteOnly);
+//////////////////////////////////////////////////////////////////////////
 
-		FMemory::Memcpy(DstVertexData, FinalVertexData.GetData(), RequiredVertexBufferSize);
-		FMemory::Memcpy(DstIndexData, FinalIndexData.GetData(), RequiredIndexBufferSize);
-
-		RHICmdList.UnlockBuffer(SourceVertexBuffer.VertexBufferRHI);
-		RHICmdList.UnlockBuffer(SourceIndexBuffer.IndexBufferRHI);
-	}
-
-	checkSlow(SourceVertexBuffer.GetBufferUsageSize() <= SourceVertexBuffer.GetBufferSize());
-	checkSlow(SourceIndexBuffer.GetBufferUsageSize() <= SourceIndexBuffer.GetBufferSize());
-
-	SET_DWORD_STAT(STAT_SlateNumLayers, InBatchData.GetNumLayers());
-	SET_DWORD_STAT(STAT_SlateNumBatches, InBatchData.GetNumFinalBatches());
-	SET_DWORD_STAT(STAT_SlateVertexCount, InBatchData.GetFinalVertexData().Num());
-}
-
-static FSceneView* CreateSceneView( FSceneViewFamilyContext* ViewFamilyContext, FSlateBackBuffer& BackBuffer, const FMatrix& ViewProjectionMatrix, const FSlateRenderingParams& Params)
+TConstArrayView<FTextureLODGroup> GetTextureLODGroups()
 {
-	QUICK_SCOPE_CYCLE_COUNTER(STAT_Slate_CreateSceneView);
-	// In loading screens, the engine is NULL, so we skip out.
-	if (GEngine == nullptr)
-	{
-		return nullptr;
-	}
-	
-	FIntRect ViewRect = Params.ViewRect;
-	if (ViewRect.IsEmpty())
-	{
-		ViewRect = FIntRect(FIntPoint(0, 0), BackBuffer.GetSizeXY());
-	}
-
-	// make a temporary view
-	FSceneViewInitOptions ViewInitOptions;
-	ViewInitOptions.ViewFamily = ViewFamilyContext;
-	ViewInitOptions.SetViewRectangle(ViewRect);
-	ViewInitOptions.ViewOrigin = FVector::ZeroVector;
-	ViewInitOptions.ViewRotationMatrix = FMatrix::Identity;
-	ViewInitOptions.ProjectionMatrix = ViewProjectionMatrix;
-	ViewInitOptions.BackgroundColor = FLinearColor::Black;
-	ViewInitOptions.OverlayColor = FLinearColor::White;
-
-	FSceneView* View = new FSceneView( ViewInitOptions );
-	ViewFamilyContext->Views.Add( View );
-
-	const FIntPoint BufferSize = BackBuffer.GetSizeXY();
-
-	// Create the view's uniform buffer.
-	FViewUniformShaderParameters ViewUniformShaderParameters;
-	ViewUniformShaderParameters.VTFeedbackBuffer = GEmptyStructuredBufferWithUAV->UnorderedAccessViewRHI;
-
-	View->SetupCommonViewUniformBufferParameters(
-		ViewUniformShaderParameters,
-		BufferSize,
-		1,
-		ViewRect,
-		View->ViewMatrices,
-		FViewMatrices()
-	);
-
-	// Update Viewport Scale UI from any external sources (Material editor, UMG zoom scale / etc).
-	ViewUniformShaderParameters.ViewportScaleUI = Params.ViewportScaleUI;
-
-	// Always Update cursor position in realtime for slate
-	ViewUniformShaderParameters.CursorPosition = Params.CursorPostion;
-
-	// TODO LWC
-	ViewUniformShaderParameters.WorldViewOriginHigh = (FVector3f)View->ViewMatrices.GetViewOrigin();
-	
-	// Slate materials need this scale to be positive, otherwise it can fail in querying scene textures (e.g., custom stencil)
-	ViewUniformShaderParameters.BufferToSceneTextureScale = FVector2f(1.0f, 1.0f);
-
-	ERHIFeatureLevel::Type RHIFeatureLevel = View->GetFeatureLevel();
-
-	ViewUniformShaderParameters.MobilePreviewMode =
-		(GIsEditor &&
-		(RHIFeatureLevel == ERHIFeatureLevel::ES3_1) &&
-		GMaxRHIFeatureLevel > ERHIFeatureLevel::ES3_1) ? 1.0f : 0.0f;
-
-	UpdateNoiseTextureParameters(ViewUniformShaderParameters);
-
+	if (UDeviceProfileManager::DeviceProfileManagerSingleton)
 	{
-		QUICK_SCOPE_CYCLE_COUNTER(STAT_Slate_CreateViewUniformBufferImmediate);
-		View->ViewUniformBuffer = TUniformBufferRef<FViewUniformShaderParameters>::CreateUniformBufferImmediate(ViewUniformShaderParameters, UniformBuffer_SingleFrame);
+		if (UDeviceProfile* Profile = UDeviceProfileManager::DeviceProfileManagerSingleton->GetActiveProfile())
+		{
+			return Profile->GetTextureLODSettings()->TextureLODGroups;
+		}
 	}
-
-	return View;
+	return {};
 }
 
-static const FName RendererModuleName("Renderer");
-
-static bool UpdateScissorRect(
-	FRHICommandList& RHICmdList, 
-#if STATS
-	int32& ScissorClips, 
-	int32& StencilClips,
-#endif
-	uint32& StencilRef, 
-	uint32& MaskingID,
-	FSlateBackBuffer& BackBuffer,
-	const FSlateRenderBatch& RenderBatch, 
-	FRHIRenderPassInfo& RPInfo,
-	FTextureRHIRef& DepthStencilTarget,
-	const FSlateClippingState*& LastClippingState,
-	const FVector2f ViewTranslation2D, 
-	FGraphicsPipelineStateInitializer& InGraphicsPSOInit,
-	FSlateStencilClipVertexBuffer& StencilVertexBuffer,
-	const FMatrix& ViewProjection, 
-	bool bForceStateChange)
+ETextureSamplerFilter GetSamplerFilter(TConstArrayView<FTextureLODGroup> TextureLODGroups, const UTexture* Texture)
 {
-	check(RHICmdList.IsInsideRenderPass());
-	bool bDidRestartRenderpass = false;
+	// Default to point filtering.
+	ETextureSamplerFilter Filter = ETextureSamplerFilter::Point;
 
-	if (RenderBatch.ClippingState != LastClippingState || bForceStateChange)
+	switch (Texture->Filter)
 	{
-		QUICK_SCOPE_CYCLE_COUNTER(STAT_Slate_UpdateScissorRect);
-
-		if (RenderBatch.ClippingState)
-		{
-			const FSlateClippingState& ClipState = *RenderBatch.ClippingState;
-			if (ClipState.GetClippingMethod() == EClippingMethod::Scissor)
-			{
-#if STATS
-				ScissorClips++;
-#endif
-
-				if (bForceStateChange && MaskingID > 0)
-				{
-					// #todo-renderpasses this is very gross. If/when this gets refactored we can detect a simple clear or batch up elements by rendertarget (and other stuff)
-					RHICmdList.EndRenderPass();
-					bDidRestartRenderpass = true;
-					ERenderTargetActions StencilAction = ERenderTargetActions::Load_Store;
-
-					RPInfo.DepthStencilRenderTarget.Action = MakeDepthStencilTargetActions(ERenderTargetActions::DontLoad_DontStore, StencilAction);
-					RPInfo.DepthStencilRenderTarget.DepthStencilTarget = DepthStencilTarget;
-					RPInfo.DepthStencilRenderTarget.ExclusiveDepthStencil = FExclusiveDepthStencil::DepthNop_StencilWrite;
-
-					RHICmdList.BeginRenderPass(RPInfo, TEXT("SlateUpdateScissorRect"));
-				}
-
-				const FSlateClippingZone& ScissorRect = ClipState.ScissorRect.GetValue();
-
-				const FIntPoint SizeXY = BackBuffer.GetSizeXY();
-				const FVector2f ViewSize((float) SizeXY.X, (float) SizeXY.Y);
-
-				// Clamp scissor rect to BackBuffer size
-				const FVector2f TopLeft     = FVector2f::Min(FVector2f::Max(ScissorRect.TopLeft     + ViewTranslation2D, FVector2f(0.0f, 0.0f)), ViewSize);
-				const FVector2f BottomRight = FVector2f::Min(FVector2f::Max(ScissorRect.BottomRight + ViewTranslation2D, FVector2f(0.0f, 0.0f)), ViewSize);
-				
-				RHICmdList.SetScissorRect(true, TopLeft.X, TopLeft.Y, BottomRight.X, BottomRight.Y);
-
-				// Disable depth/stencil testing by default
-				InGraphicsPSOInit.DepthStencilState = TStaticDepthStencilState<false, CF_Always>::GetRHI();
-				StencilRef = 0;
-			}
-			else
-			{
-#if STATS
-				StencilClips++;
-#endif
-
-				SLATE_DRAW_EVENT(RHICmdList, StencilClipping);
-
-				check(ClipState.StencilQuads.Num() > 0);
-
-				const TArray<FSlateClippingZone>& StencilQuads = ClipState.StencilQuads;
-
-				// We're going to overflow the masking ID this time, we need to reset the MaskingID to 0.
-				// this will cause us to clear the stencil buffer so that we can begin fresh.
-				if ((MaskingID + StencilQuads.Num()) > 255)
-				{
-					MaskingID = 0;
-				}
-
-				// We only clear the stencil the first time, and if some how the user draws more than 255 masking quads
-				// in a single frame.
-				const bool bClearStencil = MaskingID == 0;
-
-				// Don't bother setting the render targets unless we actually need to clear them.
-				if (bClearStencil || bForceStateChange)
-				{
-					// #todo-renderpasses Similar to above this is gross. Would require a refactor to really fix.
-					RHICmdList.EndRenderPass();
-					bDidRestartRenderpass = true;
-
-					// Clear current stencil buffer, we use ELoad/EStore, because we need to keep the stencil around.
-					ERenderTargetLoadAction StencilLoadAction = bClearStencil ? ERenderTargetLoadAction::EClear : ERenderTargetLoadAction::ELoad;
-					ERenderTargetActions StencilAction = MakeRenderTargetActions(StencilLoadAction, ERenderTargetStoreAction::EStore);
-					RPInfo.DepthStencilRenderTarget.Action = MakeDepthStencilTargetActions(ERenderTargetActions::DontLoad_DontStore, StencilAction);
-					RPInfo.DepthStencilRenderTarget.DepthStencilTarget = DepthStencilTarget;
-					RPInfo.DepthStencilRenderTarget.ExclusiveDepthStencil = FExclusiveDepthStencil::DepthNop_StencilWrite;
-					TransitionRenderPassTargets(RHICmdList, RPInfo);
-					RHICmdList.BeginRenderPass(RPInfo, TEXT("SlateUpdateScissorRect_ClearStencil"));
-				}
-
-				// Setup the scissor rect after starting the render pass, as the RHI does not preserve the scissor state between passes / render targets.
-
-				if (bClearStencil)
-				{
-					// We don't want there to be any scissor rect when we clear the stencil
-					RHICmdList.SetScissorRect(false, 0, 0, 0, 0);
-				}
-				else
-				{
-					// There might be some large - useless stencils, especially in the first couple of stencils if large
-					// widgets that clip also contain render targets, so, by setting the scissor to the AABB of the final
-					// stencil, we can cut out a lot of work that can't possibly be useful.
-					//
-					// NOTE - We also round it, because if we don't it can over-eagerly slice off pixels it shouldn't.
-					const FSlateClippingZone& MaskQuad = StencilQuads.Last();
-					const FSlateRect LastStencilBoundingBox = MaskQuad.GetBoundingBox().Round();
-
-					FSlateRect ScissorRect = LastStencilBoundingBox.OffsetBy(FVector2D(ViewTranslation2D));
-
-					// Chosen stencil quad might have some coordinates outside the viewport.
-					// After turning it into a bounding box, this box must be clamped to the current viewport,
-					// as scissors outside the viewport don't make sense (and cause assertions to fail).
-					const FIntPoint BackBufferSize = BackBuffer.GetSizeXY();
-					ScissorRect.Left = FMath::Clamp(ScissorRect.Left, 0.0f, static_cast<float>(BackBufferSize.X));
-					ScissorRect.Top = FMath::Clamp(ScissorRect.Top, 0.0f, static_cast<float>(BackBufferSize.Y));
-					ScissorRect.Right = FMath::Clamp(ScissorRect.Right, ScissorRect.Left, static_cast<float>(BackBufferSize.X));
-					ScissorRect.Bottom = FMath::Clamp(ScissorRect.Bottom, ScissorRect.Top, static_cast<float>(BackBufferSize.Y));
-
-					RHICmdList.SetScissorRect(true, ScissorRect.Left, ScissorRect.Top, ScissorRect.Right, ScissorRect.Bottom);
-				}
-
-
-				FGlobalShaderMap* MaxFeatureLevelShaderMap = GetGlobalShaderMap(GMaxRHIShaderPlatform);
-
-				// Set the new shaders
-				TShaderMapRef<FSlateMaskingVS> VertexShader(MaxFeatureLevelShaderMap);
-				TShaderMapRef<FSlateMaskingPS> PixelShader(MaxFeatureLevelShaderMap);
-
-				// Start by setting up the stenciling states so that we can write representations of the clipping zones into the stencil buffer only.
-				{
-					FGraphicsPipelineStateInitializer WriteMaskPSOInit;
-					RHICmdList.ApplyCachedRenderTargets(WriteMaskPSOInit);
-					WriteMaskPSOInit.BlendState = TStaticBlendStateWriteMask<CW_NONE, CW_NONE, CW_NONE, CW_NONE, CW_NONE, CW_NONE, CW_NONE, CW_NONE>::GetRHI();
-					WriteMaskPSOInit.RasterizerState = TStaticRasterizerState<>::GetRHI();
-					WriteMaskPSOInit.DepthStencilState =
-						TStaticDepthStencilState<
-						/*bEnableDepthWrite*/ false
-						, /*DepthTest*/ CF_Always
-						, /*bEnableFrontFaceStencil*/ true
-						, /*FrontFaceStencilTest*/ CF_Always
-						, /*FrontFaceStencilFailStencilOp*/ SO_Keep
-						, /*FrontFaceDepthFailStencilOp*/ SO_Keep
-						, /*FrontFacePassStencilOp*/ SO_Replace
-						, /*bEnableBackFaceStencil*/ true
-						, /*BackFaceStencilTest*/ CF_Always
-						, /*BackFaceStencilFailStencilOp*/ SO_Keep
-						, /*BackFaceDepthFailStencilOp*/ SO_Keep
-						, /*BackFacePassStencilOp*/ SO_Replace
-						, /*StencilReadMask*/ 0xFF
-						, /*StencilWriteMask*/ 0xFF>::GetRHI();
-
-					WriteMaskPSOInit.BoundShaderState.VertexDeclarationRHI = GSlateMaskingVertexDeclaration.VertexDeclarationRHI;
-					WriteMaskPSOInit.BoundShaderState.VertexShaderRHI = VertexShader.GetVertexShader();
-					WriteMaskPSOInit.BoundShaderState.PixelShaderRHI = PixelShader.GetPixelShader();
-					WriteMaskPSOInit.PrimitiveType = PT_TriangleStrip;
-
-					// Draw the first stencil using SO_Replace, so that we stomp any pixel with a MaskingID + 1.
-					{
-
-						SetGraphicsPipelineState(RHICmdList, WriteMaskPSOInit, MaskingID + 1);
-
-						const FSlateClippingZone& MaskQuad = StencilQuads[0];
-
-						FRHIBatchedShaderParameters& BatchedParameters = RHICmdList.GetScratchShaderParameters();
-
-						VertexShader->SetViewProjection(BatchedParameters, FMatrix44f(ViewProjection));
-						VertexShader->SetMaskRect(BatchedParameters, MaskQuad.TopLeft, MaskQuad.TopRight, MaskQuad.BottomLeft, MaskQuad.BottomRight);
-
-						RHICmdList.SetBatchedShaderParameters(VertexShader.GetVertexShader(), BatchedParameters);
-
-						RHICmdList.SetStreamSource(0, StencilVertexBuffer.VertexBufferRHI, 0);
-						RHICmdList.DrawPrimitive(0, 2, 1);
-					}
-
-					// Now setup the pipeline to use SO_SaturatedIncrement, since we've established the initial
-					// stencil with SO_Replace, we can safely use SO_SaturatedIncrement, to build up the stencil
-					// to the required mask of MaskingID + StencilQuads.Num(), thereby ensuring only the union of
-					// all stencils will render pixels.
-					{
-						WriteMaskPSOInit.DepthStencilState =
-							TStaticDepthStencilState<
-							/*bEnableDepthWrite*/ false
-							, /*DepthTest*/ CF_Always
-							, /*bEnableFrontFaceStencil*/ true
-							, /*FrontFaceStencilTest*/ CF_Always
-							, /*FrontFaceStencilFailStencilOp*/ SO_Keep
-							, /*FrontFaceDepthFailStencilOp*/ SO_Keep
-							, /*FrontFacePassStencilOp*/ SO_SaturatedIncrement
-							, /*bEnableBackFaceStencil*/ true
-							, /*BackFaceStencilTest*/ CF_Always
-							, /*BackFaceStencilFailStencilOp*/ SO_Keep
-							, /*BackFaceDepthFailStencilOp*/ SO_Keep
-							, /*BackFacePassStencilOp*/ SO_SaturatedIncrement
-							, /*StencilReadMask*/ 0xFF
-							, /*StencilWriteMask*/ 0xFF>::GetRHI();
-
-
-						SetGraphicsPipelineState(RHICmdList, WriteMaskPSOInit, 0);
-
-						FRHIBatchedShaderParameters& BatchedParameters = RHICmdList.GetScratchShaderParameters();
-						VertexShader->SetViewProjection(BatchedParameters, FMatrix44f(ViewProjection));
-						RHICmdList.SetBatchedShaderParameters(VertexShader.GetVertexShader(), BatchedParameters);
-					}
-				}
-
-				MaskingID += StencilQuads.Num();
-
-				// Next write the number of quads representing the number of clipping zones have on top of each other.
-				for (int32 MaskIndex = 1; MaskIndex < StencilQuads.Num(); MaskIndex++)
-				{
-					const FSlateClippingZone& MaskQuad = StencilQuads[MaskIndex];
-
-					FRHIBatchedShaderParameters& BatchedParameters = RHICmdList.GetScratchShaderParameters();
-
-					VertexShader->SetViewProjection(BatchedParameters, FMatrix44f(ViewProjection));
-					VertexShader->SetMaskRect(BatchedParameters, MaskQuad.TopLeft, MaskQuad.TopRight, MaskQuad.BottomLeft, MaskQuad.BottomRight);
-
-					RHICmdList.SetBatchedShaderParameters(VertexShader.GetVertexShader(), BatchedParameters);
-
-					RHICmdList.SetStreamSource(0, StencilVertexBuffer.VertexBufferRHI, 0);
-					RHICmdList.DrawPrimitive(0, 2, 1);
-				}
-
-				// Setup the stenciling state to be read only now, disable depth writes, and restore the color buffer
-				// because we're about to go back to rendering widgets "normally", but with the added effect that now
-				// we have the stencil buffer bound with a bunch of clipping zones rendered into it.
-				{
-					FRHIDepthStencilState* DSMaskRead =
-						TStaticDepthStencilState<
-						/*bEnableDepthWrite*/ false
-						, /*DepthTest*/ CF_Always
-						, /*bEnableFrontFaceStencil*/ true
-						, /*FrontFaceStencilTest*/ CF_Equal
-						, /*FrontFaceStencilFailStencilOp*/ SO_Keep
-						, /*FrontFaceDepthFailStencilOp*/ SO_Keep
-						, /*FrontFacePassStencilOp*/ SO_Keep
-						, /*bEnableBackFaceStencil*/ true
-						, /*BackFaceStencilTest*/ CF_Equal
-						, /*BackFaceStencilFailStencilOp*/ SO_Keep
-						, /*BackFaceDepthFailStencilOp*/ SO_Keep
-						, /*BackFacePassStencilOp*/ SO_Keep
-						, /*StencilReadMask*/ 0xFF
-						, /*StencilWriteMask*/ 0xFF>::GetRHI();
-
-					InGraphicsPSOInit.DepthStencilState = DSMaskRead;
-
-					// We set a StencilRef equal to the number of stenciling/clipping masks,
-					// so unless the pixel we're rendering two is on top of a stencil pixel with the same number
-					// it's going to get rejected, thereby clipping everything except for the cross-section of
-					// all the stenciling quads.
-					StencilRef = MaskingID;
-				}
-			}
+	case TF_Nearest: 
+		Filter = ETextureSamplerFilter::Point; 
+		break;
+	case TF_Bilinear:
+		Filter = ETextureSamplerFilter::Bilinear; 
+		break;
+	case TF_Trilinear: 
+		Filter = ETextureSamplerFilter::Trilinear; 
+		break;
 
-			RHICmdList.ApplyCachedRenderTargets(InGraphicsPSOInit);
-		}
-		else
+		// TF_Default
+	default:
+		// Use LOD group value to find proper filter setting.
+		if (Texture->LODGroup < TextureLODGroups.Num())
 		{
-			RHICmdList.SetScissorRect(false, 0, 0, 0, 0);
-
-			// Disable depth/stencil testing
-			InGraphicsPSOInit.DepthStencilState = TStaticDepthStencilState<false, CF_Always>::GetRHI();
-			StencilRef = 0;
+			Filter = TextureLODGroups[Texture->LODGroup].Filter;
 		}
-
-		LastClippingState = RenderBatch.ClippingState;
 	}
 
-	return bDidRestartRenderpass;
-}
-
-static bool UpdateScissorRect(
-	FRHICommandList& RHICmdList,
-#if STATS
-	int32& ScissorClips,
-	int32& StencilClips,
-#endif
-	uint32& StencilRef,
-	uint32& MaskingID,
-	FSlateBackBuffer& BackBuffer,
-	const FSlateRenderBatch& RenderBatch,
-	FTextureRHIRef& ColorTarget,
-	FTextureRHIRef& DepthStencilTarget,
-	const FSlateClippingState*& LastClippingState,
-	const FVector2f ViewTranslation2D,
-	FGraphicsPipelineStateInitializer& InGraphicsPSOInit,
-	FSlateStencilClipVertexBuffer& StencilVertexBuffer,
-	const FMatrix& ViewProjection,
-	bool bForceStateChange)
-{
-	FRHIRenderPassInfo RPInfo(ColorTarget, ERenderTargetActions::Load_Store);
-	return UpdateScissorRect(RHICmdList,
-#if STATS
-		ScissorClips,
-		StencilClips,
-#endif
-		StencilRef,
-		MaskingID,
-		BackBuffer,
-		RenderBatch,
-		RPInfo,
-		DepthStencilTarget,
-		LastClippingState,
-		ViewTranslation2D,
-		InGraphicsPSOInit,
-		StencilVertexBuffer,
-		ViewProjection,
-		bForceStateChange);
+	return Filter;
 }
 
-static FRHISamplerState* GetSamplerState(ESlateBatchDrawFlag DrawFlags, ETextureSamplerFilter Filter)
+FRHISamplerState* GetSamplerState(ESlateBatchDrawFlag DrawFlags, ETextureSamplerFilter Filter)
 {
 	FRHISamplerState* SamplerState = nullptr;
 
@@ -659,766 +218,46 @@ static FRHISamplerState* GetSamplerState(ESlateBatchDrawFlag DrawFlags, ETexture
 	return SamplerState;
 }
 
-void FSlateRHIRenderingPolicy::DrawElements(
-	FRHICommandListImmediate& RHICmdList,
-	FSlateBackBuffer& BackBuffer,
-	FTextureRHIRef& ColorTarget,
-	FTextureRHIRef& PostProcessTexture,
-	FTextureRHIRef& DepthStencilTarget,
-	int32 FirstBatchIndex,
-	const TArray<FSlateRenderBatch>& RenderBatches,
-	const FSlateRenderingParams& Params)
+/** Returns the pixel shader that should be used for the specified ShaderType and DrawEffects */
+TShaderRef<FSlateElementPS> GetTexturePixelShader(FGlobalShaderMap* ShaderMap, ESlateShader ShaderType, ESlateDrawEffect DrawEffects, bool bUseTextureGrayscale, bool bIsVirtualTexture)
 {
-	// Should only be called by the rendering thread
-	check(IsInRenderingThread());
-	check(RHICmdList.IsInsideRenderPass());
-
-	// Cache the TextureLODGroups so that we can look them up for texture filtering.
-	if (UDeviceProfileManager::DeviceProfileManagerSingleton)
-	{
-		if (UDeviceProfile* Profile = UDeviceProfileManager::Get().GetActiveProfile())
-		{
-			if (Profile != LastDeviceProfile)
-			{
-				TextureLODGroups = Profile->GetTextureLODSettings()->TextureLODGroups;
-				LastDeviceProfile = Profile;
-			}
-		}
-	}
-
-	IRendererModule& RendererModule = FModuleManager::GetModuleChecked<IRendererModule>(RendererModuleName);
-
-	static const FEngineShowFlags DefaultShowFlags(ESFIM_Game);
-
-	// Disable gammatization when back buffer is in float 16 format.
-	// Note that the final editor rendering won't compare 1:1 with 8/10 bit RGBA since blending
-	// of "manually" gammatized values is wrong as there is no de-gammatization of the destination buffer
-	// and re-gammatization of the resulting blending operation in the 8/10 bit RGBA path.
-	// For Editor running in HDR then the gamma needs to be 2.2 and have a float back buffer format.
-
-	const float EngineGamma = (!GIsEditor && (BackBuffer.GetRenderTargetTexture()->GetFormat() == PF_FloatRGBA) && (Params.bIsHDR==false)) ? 1.0f : GEngine ? GEngine->GetDisplayGamma() : 2.2f;
-	const float DisplayGamma = bGammaCorrect ? EngineGamma : 1.0f;
-	const float DisplayContrast = GSlateContrast;
-
-	int32 ScissorClips = 0;
-	int32 StencilClips = 0;
-
-	// In order to support MaterialParameterCollections, we need to create multiple FSceneViews for 
-	// each possible Scene that we encounter. The following code creates these as separate arrays, where the first 
-	// N entries map directly to entries from ActiveScenes. The final entry is added to represent the absence of a
-	// valid scene, i.e. a -1 in the SceneIndex parameter of the batch.
-	int32 NumScenes = ResourceManager->GetSceneCount() + 1;
-	TArray<FSceneView*, TInlineAllocator<3> > SceneViews;
-	SceneViews.SetNum(NumScenes);
-	TArray<FSceneViewFamilyContext*, TInlineAllocator<3> > SceneViewFamilyContexts;
-	SceneViewFamilyContexts.SetNum(NumScenes);
+	TShaderRef<FSlateElementPS> PixelShader;
 
+#if WITH_SLATE_VISUALIZERS
+	if (CVarShowSlateOverdraw.GetValueOnRenderThread() != 0)
 	{
-		QUICK_SCOPE_CYCLE_COUNTER(STAT_Slate_CreateScenes);
-		for (int32 i = 0; i < ResourceManager->GetSceneCount(); i++)
-		{
-			SceneViewFamilyContexts[i] = new FSceneViewFamilyContext
-			(
-				FSceneViewFamily::ConstructionValues
-				(
-					&BackBuffer,
-					ResourceManager->GetSceneAt(i),
-					DefaultShowFlags
-				)
-				.SetTime(Params.Time)
-				.SetRealtimeUpdate(true)
-			);
-			SceneViews[i] = CreateSceneView(SceneViewFamilyContexts[i], BackBuffer, FMatrix(Params.ViewProjectionMatrix), Params);
-		}
-
-		SceneViewFamilyContexts[NumScenes - 1] = new FSceneViewFamilyContext
-		(
-			FSceneViewFamily::ConstructionValues
-			(
-				&BackBuffer,
-				nullptr,
-				DefaultShowFlags
-			)
-			.SetTime(Params.Time)
-			.SetRealtimeUpdate(true)
-		);
-		SceneViews[NumScenes - 1] = CreateSceneView(SceneViewFamilyContexts[NumScenes - 1], BackBuffer, FMatrix(Params.ViewProjectionMatrix), Params);
+		PixelShader = TShaderMapRef<FSlateDebugOverdrawPS>(ShaderMap);
 	}
-
-	TShaderMapRef<FSlateElementVS> GlobalVertexShader(GetGlobalShaderMap(GMaxRHIShaderPlatform));
-
-	FSamplerStateRHIRef BilinearClamp = TStaticSamplerState<SF_Bilinear, AM_Clamp, AM_Clamp, AM_Clamp>::GetRHI();
-
-	TSlateElementVertexBuffer<FSlateVertex>* VertexBufferPtr = &SourceVertexBuffer;
-	FSlateElementIndexBuffer* IndexBufferPtr = &SourceIndexBuffer;
-
-	FGraphicsPipelineStateInitializer GraphicsPSOInit;
-	RHICmdList.ApplyCachedRenderTargets(GraphicsPSOInit);
-
-	const FSlateRenderDataHandle* LastHandle = nullptr;
-
-	const ERHIFeatureLevel::Type FeatureLevel = GMaxRHIFeatureLevel;
-	FGlobalShaderMap* ShaderMap = GetGlobalShaderMap(GMaxRHIShaderPlatform);
-
-#if WITH_SLATE_VISUALIZERS
-	FRandomStream BatchColors(1337);
-#endif
-
-	// This variable tracks the last clipping state, so that if multiple batches have the same clipping state, we don't have to do any work.
-	const FSlateClippingState* LastClippingState;
-
-	// This is the stenciling ref variable we set any time we draw, so that any stencil comparisons use the right mask id.
-	uint32 StencilRef = 0;
-	// This is an accumulating maskID that we use to track the between batch usage of the stencil buffer, when at 0, or over 255
-	// this signals that we need to reset the masking ID, and clear the stencil buffer, as we've used up the available scratch range.
-	uint32 MaskingID = 0;
-
-	RHICmdList.SetScissorRect(false, 0, 0, 0, 0);
-	// Disable depth/stencil testing by default
-	GraphicsPSOInit.DepthStencilState = TStaticDepthStencilState<false, CF_Always>::GetRHI();
-
-	FVector2f ViewTranslation2D = Params.ViewOffset;
-
-	// Draw each element
-#if WITH_SLATE_DEBUGGING
-	int32 NextRenderBatchIndex = BatchToDraw == -1 ? FirstBatchIndex : BatchToDraw;
-#else
-	int32 NextRenderBatchIndex = FirstBatchIndex;
+	else
 #endif
-
-
-	/*
-		#todo-renderpasses This loop ends up with ugly logic.
-		CustomDrawers will draw in their own renderpass. So we must remember to reopen the renderpass with the passed in Color/DepthStencil targets.
-	*/
-	while (NextRenderBatchIndex != INDEX_NONE)
 	{
-		VertexBufferPtr = &SourceVertexBuffer;
-		IndexBufferPtr = &SourceIndexBuffer;
-
-		if (!RHICmdList.IsInsideRenderPass())
-		{
-			// Restart the renderpass since the CustomDrawer or post-process may have changed it in last iteration
-			FRHIRenderPassInfo RPInfo(BackBuffer.GetRenderTargetTexture(), ERenderTargetActions::Load_Store);
-			RPInfo.DepthStencilRenderTarget.DepthStencilTarget = DepthStencilTarget;
-			if (DepthStencilTarget)
-			{
-				RPInfo.DepthStencilRenderTarget.Action = EDepthStencilTargetActions::LoadDepthStencil_StoreDepthStencil;
-				RPInfo.DepthStencilRenderTarget.ExclusiveDepthStencil = FExclusiveDepthStencil::DepthWrite_StencilWrite;
-			}
-			else
-			{
-				RPInfo.DepthStencilRenderTarget.Action = EDepthStencilTargetActions::DontLoad_DontStore;
-				RPInfo.DepthStencilRenderTarget.ExclusiveDepthStencil = FExclusiveDepthStencil::DepthNop_StencilNop;
-			}
-			TransitionRenderPassTargets(RHICmdList, RPInfo);
-			RHICmdList.BeginRenderPass(RPInfo, TEXT("RestartingSlateDrawElements"));
-
-			// Something may have messed with the viewport size so set it back to the full target.
-			RHICmdList.SetViewport(0.f, 0.f, 0.f, (float)BackBuffer.GetSizeXY().X, (float)BackBuffer.GetSizeXY().Y, 0.0f);
-
-			// Re-apply render target states to the PSO initializer, since we've changed the depth/stencil target.
-			RHICmdList.ApplyCachedRenderTargets(GraphicsPSOInit);
-		}
-				
-#if WITH_SLATE_VISUALIZERS
-		FLinearColor BatchColor = FLinearColor(BatchColors.GetUnitVector());
-#endif
-		const FSlateRenderBatch& RenderBatch = RenderBatches[NextRenderBatchIndex];
-
-		NextRenderBatchIndex = RenderBatch.NextBatchIndex;
-
-#if WITH_SLATE_DEBUGGING
-		if (BatchToDraw != -1)
-		{
-			break;
-		}
-#endif
-
-		const FSlateShaderResource* ShaderResource = RenderBatch.ShaderResource;
-		const ESlateBatchDrawFlag DrawFlags = RenderBatch.DrawFlags;
-		const ESlateDrawEffect DrawEffects = RenderBatch.DrawEffects;
-		const ESlateShader ShaderType = RenderBatch.ShaderType;
-		const FShaderParams& ShaderParams = RenderBatch.ShaderParams;
-
-		if (EnumHasAllFlags(DrawFlags, ESlateBatchDrawFlag::Wireframe))
-		{
-			GraphicsPSOInit.RasterizerState = TStaticRasterizerState<FM_Wireframe>::GetRHI();
-		}
-		else
-		{
-			GraphicsPSOInit.RasterizerState = TStaticRasterizerState<FM_Solid>::GetRHI();
-		}
+	const bool bDrawDisabled = EnumHasAllFlags( DrawEffects, ESlateDrawEffect::DisabledEffect );
+	const bool bUseTextureAlpha = !EnumHasAllFlags( DrawEffects, ESlateDrawEffect::IgnoreTextureAlpha );
 
-		if (!RenderBatch.CustomDrawer)
+	if (bDrawDisabled)
+	{
+		switch (ShaderType)
 		{
-			FMatrix DynamicOffset = FTranslationMatrix::Make(FVector(RenderBatch.DynamicOffset.X, RenderBatch.DynamicOffset.Y, 0));
-			const FMatrix ViewProjection = DynamicOffset * FMatrix(Params.ViewProjectionMatrix);
-
-			UpdateScissorRect(
-				RHICmdList,
-#if STATS
-				ScissorClips,
-				StencilClips,
-#endif
-				StencilRef,
-				MaskingID,
-				BackBuffer,
-				RenderBatch,
-				ColorTarget,
-				DepthStencilTarget,
-				LastClippingState,
-				ViewTranslation2D,
-				GraphicsPSOInit,
-				StencilVertexBuffer,
-				ViewProjection,
-				false);
-
-			const uint32 PrimitiveCount = RenderBatch.DrawPrimitiveType == ESlateDrawPrimitive::LineList ? RenderBatch.NumIndices / 2 : RenderBatch.NumIndices / 3;
-			check(ShaderResource == nullptr || !ShaderResource->Debug_IsDestroyed());
-			ESlateShaderResource::Type ResourceType = ShaderResource ? ShaderResource->GetType() : ESlateShaderResource::Invalid;
-			if (ResourceType != ESlateShaderResource::Material && ShaderType != ESlateShader::PostProcess)
+		default:
+		case ESlateShader::Default:
+			if (bUseTextureAlpha)
 			{
-				check(RHICmdList.IsInsideRenderPass());
-				check(RenderBatch.NumIndices > 0);
-				TShaderRef<FSlateElementPS> PixelShader;
-
-				const bool bUseInstancing = RenderBatch.InstanceCount > 1 && RenderBatch.InstanceData != nullptr;
-				check(bUseInstancing == false);
-
-#if WITH_SLATE_VISUALIZERS
-				TShaderRef<FSlateDebugBatchingPS> BatchingPixelShader;
-				if (CVarShowSlateBatching.GetValueOnRenderThread() != 0)
+				if (bIsVirtualTexture)
 				{
-					BatchingPixelShader = TShaderMapRef<FSlateDebugBatchingPS>(ShaderMap);
-					PixelShader = BatchingPixelShader;
+					if (bUseTextureGrayscale)
+					{
+						PixelShader = TShaderMapRef<TSlateElementPS<ESlateShader::Default, true, true, true, true> >(ShaderMap);
+					}
+					else
+					{
+						PixelShader = TShaderMapRef<TSlateElementPS<ESlateShader::Default, true, true, false, true> >(ShaderMap);
+					}
 				}
 				else
-#endif
 				{
-					bool bIsVirtualTexture = false;
-
-					// check if texture is using BC4 compression and set shader to render grayscale
-					bool bUseTextureGrayscale = false;
-
-
-					if ((ShaderResource != nullptr) && (ResourceType == ESlateShaderResource::TextureObject))
+					if (bUseTextureGrayscale)
 					{
-						FSlateBaseUTextureResource* TextureObjectResource = const_cast<FSlateBaseUTextureResource*>(static_cast<const FSlateBaseUTextureResource*>(ShaderResource));
-						
-						if (UTexture* TextureObj = TextureObjectResource->GetTextureObject())
-						{
-							bIsVirtualTexture = TextureObj->IsCurrentlyVirtualTextured();
-
-							if (TextureObj->CompressionSettings == TC_Alpha)
-							{
-								bUseTextureGrayscale = true;
-							}
-						}
-					}
-
-					PixelShader = GetTexturePixelShader(ShaderMap, ShaderType, DrawEffects, bUseTextureGrayscale, bIsVirtualTexture);
-				}
-
-#if WITH_SLATE_VISUALIZERS
-				if (CVarShowSlateBatching.GetValueOnRenderThread() != 0)
-				{
-					GraphicsPSOInit.BlendState = TStaticBlendState<CW_RGBA, BO_Add, BF_SourceAlpha, BF_InverseSourceAlpha, BO_Add, BF_One, BF_InverseSourceAlpha>::GetRHI();
-				}
-				else if (CVarShowSlateOverdraw.GetValueOnRenderThread() != 0)
-				{
-					GraphicsPSOInit.BlendState = TStaticBlendState<CW_RGB, BO_Add, BF_One, BF_One, BO_Add, BF_Zero, BF_InverseSourceAlpha>::GetRHI();
-				}
-				else
-#endif
-				{
-					GraphicsPSOInit.BlendState =
-						EnumHasAllFlags(DrawFlags, ESlateBatchDrawFlag::NoBlending)
-						? TStaticBlendState<>::GetRHI()
-						: (EnumHasAllFlags(DrawFlags, ESlateBatchDrawFlag::PreMultipliedAlpha)
-							? TStaticBlendState<CW_RGBA, BO_Add, BF_One, BF_InverseSourceAlpha, BO_Add, BF_One, BF_InverseSourceAlpha>::GetRHI()
-							: TStaticBlendState<CW_RGBA, BO_Add, BF_SourceAlpha, BF_InverseSourceAlpha, BO_Add, BF_One, BF_InverseSourceAlpha>::GetRHI())
-						;
-				}
-
-				if (EnumHasAllFlags(DrawFlags, ESlateBatchDrawFlag::Wireframe) || Params.bWireFrame)
-				{
-					GraphicsPSOInit.RasterizerState = TStaticRasterizerState<FM_Wireframe>::GetRHI();
-
-					if (Params.bWireFrame)
-					{
-						GraphicsPSOInit.BlendState = TStaticBlendState<>::GetRHI();
-					}
-				}
-				else
-				{
-					GraphicsPSOInit.RasterizerState = TStaticRasterizerState<FM_Solid>::GetRHI();
-				}
-
-				GraphicsPSOInit.BoundShaderState.VertexDeclarationRHI = GSlateVertexDeclaration.VertexDeclarationRHI;
-				GraphicsPSOInit.BoundShaderState.VertexShaderRHI = GlobalVertexShader.GetVertexShader();
-				GraphicsPSOInit.BoundShaderState.PixelShaderRHI = PixelShader.GetPixelShader();
-				GraphicsPSOInit.PrimitiveType = GetRHIPrimitiveType(RenderBatch.DrawPrimitiveType);
-
-				SetGraphicsPipelineState(RHICmdList, GraphicsPSOInit, StencilRef);
-
-				FRHIBatchedShaderParameters& BatchedParameters = RHICmdList.GetScratchShaderParameters();
-
-#if WITH_SLATE_VISUALIZERS
-				if (CVarShowSlateBatching.GetValueOnRenderThread() != 0)
-				{
-					BatchingPixelShader->SetBatchColor(BatchedParameters, BatchColor);
-				}
-#endif
-
-				FRHISamplerState* SamplerState = BilinearClamp;
-				FRHITexture* TextureRHI = GWhiteTexture->TextureRHI;
-				bool bIsVirtualTexture = false;
-				FTextureResource* TextureResource = nullptr;
-
-				if (ShaderResource)
-				{
-					ETextureSamplerFilter Filter = ETextureSamplerFilter::Bilinear;
-
-					if (ResourceType == ESlateShaderResource::TextureObject)
-					{
-						FSlateBaseUTextureResource* TextureObjectResource = (FSlateBaseUTextureResource*)ShaderResource;
-						if (UTexture* TextureObj = TextureObjectResource->GetTextureObject())
-						{
-							TextureObjectResource->CheckForStaleResources();
-
-							TextureRHI = TextureObjectResource->AccessRHIResource();
-
-							// This can upset some RHIs, so use transparent black texture until it's valid.
-							// these can be temporarily invalid when recreating them / invalidating their streaming
-							// state.
-							if (TextureRHI == nullptr)
-							{
-								// We use transparent black here, because it's about to become valid - probably, and flashing white
-								// wouldn't be ideal.
-								TextureRHI = GTransparentBlackTexture->TextureRHI;
-							}
-
-							TextureResource = TextureObj->GetResource();
-
-							Filter = GetSamplerFilter(TextureObj);
-							bIsVirtualTexture = TextureObj->IsCurrentlyVirtualTextured();
-						}
-					}
-					else
-					{
-						FRHITexture* NativeTextureRHI = ((TSlateTexture<FTextureRHIRef>*)ShaderResource)->GetTypedResource();
-						// Atlas textures that have no content are never initialized but null textures are invalid on many platforms.
-						TextureRHI = NativeTextureRHI ? NativeTextureRHI : (FRHITexture*)GWhiteTexture->TextureRHI;
-					}
-
-					SamplerState = GetSamplerState(DrawFlags, Filter);
-				}
-
-				{
-					if (bIsVirtualTexture && (TextureResource != nullptr))
-					{
-						PixelShader->SetVirtualTextureParameters(BatchedParameters, static_cast<FVirtualTexture2DResource*>(TextureResource));
-					}
-					else
-					{
-						PixelShader->SetTexture(BatchedParameters, TextureRHI, SamplerState);
-					}
-					
-					PixelShader->SetShaderParams(BatchedParameters, ShaderParams);
-					const float FinalGamma = EnumHasAnyFlags(DrawFlags, ESlateBatchDrawFlag::ReverseGamma) ? (1.0f / EngineGamma) : EnumHasAnyFlags(DrawFlags, ESlateBatchDrawFlag::NoGamma) ? 1.0f : DisplayGamma;
-					const float FinalContrast = EnumHasAnyFlags(DrawFlags, ESlateBatchDrawFlag::NoGamma) ? 1 : DisplayContrast;
-					PixelShader->SetDisplayGammaAndInvertAlphaAndContrast(BatchedParameters, FinalGamma, EnumHasAllFlags(DrawEffects, ESlateDrawEffect::InvertAlpha) ? 1.0f : 0.0f, FinalContrast);
-
-					RHICmdList.SetBatchedShaderParameters(PixelShader.GetPixelShader(), BatchedParameters);
-				}
-				{
-					GlobalVertexShader->SetViewProjection(BatchedParameters, FMatrix44f(ViewProjection));
-					RHICmdList.SetBatchedShaderParameters(GlobalVertexShader.GetVertexShader(), BatchedParameters);
-				}
-
-				{
-					// for RHIs that can't handle VertexOffset, we need to offset the stream source each time
-					RHICmdList.SetStreamSource(0, VertexBufferPtr->VertexBufferRHI, RenderBatch.VertexOffset * sizeof(FSlateVertex));
-					RHICmdList.DrawIndexedPrimitive(IndexBufferPtr->IndexBufferRHI, 0, 0, RenderBatch.NumVertices, RenderBatch.IndexOffset, PrimitiveCount, RenderBatch.InstanceCount);
-				}
-			}
-			else if (GEngine && ShaderResource && ShaderResource->GetType() == ESlateShaderResource::Material && ShaderType != ESlateShader::PostProcess)
-			{
-				check(RHICmdList.IsInsideRenderPass());
-
-				check(RenderBatch.NumIndices > 0);
-				// Note: This code is only executed if the engine is loaded (in early loading screens attempting to use a material is unsupported
-				int32 ActiveSceneIndex = (int32)RenderBatch.SceneIndex;
-
-				// We are assuming at this point that the SceneIndex from the batch is either -1, meaning no scene or a valid scene.
-				// We set up the "no scene" option as the last SceneView in the array above.
-				if (RenderBatch.SceneIndex == -1)
-				{
-					ActiveSceneIndex = NumScenes - 1;
-				}
-				else if (RenderBatch.SceneIndex >= ResourceManager->GetSceneCount())
-				{
-					// Ideally we should never hit this scenario, but given that Paragon may be using cached
-					// render batches and is running into this daily, for this branch we should
-					// just ignore the scene if the index is invalid. Note that the
-					// MaterialParameterCollections will not be correct for this scene, should they be
-					// used.
-					ActiveSceneIndex = NumScenes - 1;
-#if UE_BUILD_DEBUG && WITH_EDITOR
-					UE_LOG(LogSlate, Error, TEXT("Invalid scene index in batch: %d of %d known scenes!"), RenderBatch.SceneIndex, ResourceManager->GetSceneCount());
-#endif
-				}
-
-				// Handle the case where we skipped out early above
-				if (SceneViews[ActiveSceneIndex] == nullptr)
-				{
-					continue;
-				}
-
-				const FSceneView& ActiveSceneView = *SceneViews[ActiveSceneIndex];
-
-				FSlateMaterialResource* MaterialShaderResource = (FSlateMaterialResource*)ShaderResource;
-				if (const FMaterialRenderProxy* MaterialRenderProxy = MaterialShaderResource->GetRenderProxy())
-				{
-					SLATE_DRAW_EVENTF(RHICmdList, MaterialBatch, TEXT("Slate Material: %s"), MaterialRenderProxy->GetMaterialName());
-
-					const ERHIFeatureLevel::Type ViewFeatureLevel = ActiveSceneView.GetFeatureLevel();
-					// Update Material in case of null ShaderMap
-					MaterialRenderProxy->UpdateUniformExpressionCacheIfNeeded(RHICmdList, ViewFeatureLevel);
-					MaterialShaderResource->CheckForStaleResources();
-
-					const bool bUseInstancing = RenderBatch.InstanceCount > 0 && RenderBatch.InstanceData != nullptr;
-
-					TShaderRef<FSlateMaterialShaderVS> VertexShader;
-					TShaderRef<FSlateMaterialShaderPS> PixelShader;
-
-					FMaterialShaderTypes ShaderTypesToGet;
-					ChooseMaterialShaderTypes(ShaderType, bUseInstancing, ShaderTypesToGet);
-					const FMaterial* EffectiveMaterial = nullptr;
-
-					while(MaterialRenderProxy)
-					{
-						const FMaterial* Material = MaterialRenderProxy->GetMaterialNoFallback(ViewFeatureLevel);
-						FMaterialShaders Shaders;
-						if (Material && Material->TryGetShaders(ShaderTypesToGet, nullptr, Shaders))
-						{
-							EffectiveMaterial = Material;
-							Shaders.TryGetVertexShader(VertexShader);
-							Shaders.TryGetPixelShader(PixelShader);
-							break;
-						}
-
-						MaterialRenderProxy = MaterialRenderProxy->GetFallback(ViewFeatureLevel);
-					}
-
-					FRHIUniformBuffer* SceneTextureUniformBuffer = GetSceneTextureExtracts().GetUniformBuffer();
-
-					if (VertexShader.IsValid() && PixelShader.IsValid() && SceneTextureUniformBuffer)
-					{
-						check(EffectiveMaterial);
-						const FUniformBufferStaticBindings StaticUniformBuffers(SceneTextureUniformBuffer);
-						SCOPED_UNIFORM_BUFFER_STATIC_BINDINGS(RHICmdList, StaticUniformBuffers);
-
-						FRHIBatchedShaderParameters& BatchedParameters = RHICmdList.GetScratchShaderParameters();
-
-#if WITH_SLATE_VISUALIZERS
-						if (CVarShowSlateBatching.GetValueOnRenderThread() != 0)
-						{
-							TShaderMapRef<FSlateDebugBatchingPS> BatchingPixelShader(ShaderMap);
-
-							GraphicsPSOInit.BoundShaderState.VertexDeclarationRHI = bUseInstancing ? GSlateInstancedVertexDeclaration.VertexDeclarationRHI : GSlateVertexDeclaration.VertexDeclarationRHI;
-							GraphicsPSOInit.BoundShaderState.VertexShaderRHI = GlobalVertexShader.GetVertexShader();
-							GraphicsPSOInit.BoundShaderState.PixelShaderRHI = BatchingPixelShader.GetPixelShader();
-							GraphicsPSOInit.BlendState = TStaticBlendState<CW_RGB, BO_Add, BF_One, BF_One, BO_Add, BF_Zero, BF_InverseSourceAlpha>::GetRHI();
-
-							BatchingPixelShader->SetBatchColor(BatchedParameters, BatchColor);
-						}
-						else if (CVarShowSlateOverdraw.GetValueOnRenderThread() != 0)
-						{
-							TShaderMapRef<FSlateDebugOverdrawPS> OverdrawPixelShader(ShaderMap);
-
-							GraphicsPSOInit.BoundShaderState.VertexDeclarationRHI = bUseInstancing ? GSlateInstancedVertexDeclaration.VertexDeclarationRHI : GSlateVertexDeclaration.VertexDeclarationRHI;
-							GraphicsPSOInit.BoundShaderState.VertexShaderRHI = GlobalVertexShader.GetVertexShader();
-							GraphicsPSOInit.BoundShaderState.PixelShaderRHI = OverdrawPixelShader.GetPixelShader();
-							GraphicsPSOInit.BlendState = TStaticBlendState<CW_RGB, BO_Add, BF_One, BF_One, BO_Add, BF_Zero, BF_InverseSourceAlpha>::GetRHI();
-						}
-#endif
-						{
-							PixelShader->SetBlendState(GraphicsPSOInit, EffectiveMaterial);
-							FSlateShaderResource* MaskResource = MaterialShaderResource->GetTextureMaskResource();
-							if (MaskResource && IsOpaqueOrMaskedBlendMode(*EffectiveMaterial))
-							{
-								// Font materials require some form of translucent blending
-								GraphicsPSOInit.BlendState = TStaticBlendState<CW_RGBA, BO_Add, BF_SourceAlpha, BF_InverseSourceAlpha, BO_Add, BF_InverseDestAlpha, BF_One>::GetRHI();
-							}
-
-							GraphicsPSOInit.BoundShaderState.VertexDeclarationRHI = bUseInstancing ? GSlateInstancedVertexDeclaration.VertexDeclarationRHI : GSlateVertexDeclaration.VertexDeclarationRHI;
-							GraphicsPSOInit.BoundShaderState.VertexShaderRHI = VertexShader.GetVertexShader();
-							GraphicsPSOInit.BoundShaderState.PixelShaderRHI = PixelShader.GetPixelShader();
-							GraphicsPSOInit.PrimitiveType = GetRHIPrimitiveType(RenderBatch.DrawPrimitiveType);
-
-							SetGraphicsPipelineState(RHICmdList, GraphicsPSOInit, StencilRef);
-
-							{
-
-								PixelShader->SetParameters(BatchedParameters, ActiveSceneView, MaterialRenderProxy, EffectiveMaterial, ShaderParams);
-								const float FinalGamma = EnumHasAnyFlags(DrawFlags, ESlateBatchDrawFlag::ReverseGamma) ? 1.0f / EngineGamma : EnumHasAnyFlags(DrawFlags, ESlateBatchDrawFlag::NoGamma) ? 1.0f : DisplayGamma;
-								const float FinalContrast = EnumHasAnyFlags(DrawFlags, ESlateBatchDrawFlag::NoGamma) ? 1 : DisplayContrast;
-								PixelShader->SetDisplayGammaAndContrast(BatchedParameters, FinalGamma, FinalContrast);
-								const bool bDrawDisabled = EnumHasAllFlags(DrawEffects, ESlateDrawEffect::DisabledEffect);
-								PixelShader->SetDrawFlags(BatchedParameters, bDrawDisabled);
-
-								if (MaskResource)
-								{
-									FTextureRHIRef TextureRHI;
-									TextureRHI = ((TSlateTexture<FTextureRHIRef>*)MaskResource)->GetTypedResource();
-
-									PixelShader->SetAdditionalTexture(BatchedParameters, TextureRHI, BilinearClamp);
-								}
-
-								RHICmdList.SetBatchedShaderParameters(PixelShader.GetPixelShader(), BatchedParameters);
-							}
-							{
-								VertexShader->SetViewProjection(BatchedParameters, FMatrix44f(ViewProjection));
-								VertexShader->SetMaterialShaderParameters(BatchedParameters, ActiveSceneView, MaterialRenderProxy, EffectiveMaterial);
-								RHICmdList.SetBatchedShaderParameters(VertexShader.GetVertexShader(), BatchedParameters);
-							}
-						}
-
-						{
-							if (bUseInstancing)
-							{
-								uint32 InstanceCount = RenderBatch.InstanceCount;
-
-								RenderBatch.InstanceData->BindStreamSource(RHICmdList, 1, RenderBatch.InstanceOffset);
-
-								// for RHIs that can't handle VertexOffset, we need to offset the stream source each time
-
-								RHICmdList.SetStreamSource(0, VertexBufferPtr->VertexBufferRHI, RenderBatch.VertexOffset * sizeof(FSlateVertex));
-								RHICmdList.DrawIndexedPrimitive(IndexBufferPtr->IndexBufferRHI, 0, 0, RenderBatch.NumVertices, RenderBatch.IndexOffset, PrimitiveCount, InstanceCount);
-							}
-							else
-							{
-								RHICmdList.SetStreamSource(1, nullptr, 0);
-
-								// for RHIs that can't handle VertexOffset, we need to offset the stream source each time
-								RHICmdList.SetStreamSource(0, VertexBufferPtr->VertexBufferRHI, RenderBatch.VertexOffset * sizeof(FSlateVertex));
-								RHICmdList.DrawIndexedPrimitive(IndexBufferPtr->IndexBufferRHI, 0, 0, RenderBatch.NumVertices, RenderBatch.IndexOffset, PrimitiveCount, 1);
-
-							}
-						}
-					}
-				}
-			}
-			else if (ShaderType == ESlateShader::PostProcess)
-			{
-				SLATE_DRAW_EVENT(RHICmdList, PostProcess);
-				RHICmdList.EndRenderPass();
-
-				const FVector4f QuadPositionData = ShaderParams.PixelParams;
-
-				FPostProcessRectParams RectParams;
-				RectParams.SourceTexture = PostProcessTexture;
-				RectParams.SourceRect = FSlateRect((float)Params.ViewRect.Min.X, (float)Params.ViewRect.Min.Y, (float)Params.ViewRect.Max.X, (float)Params.ViewRect.Max.Y);
-				RectParams.DestRect = FSlateRect(QuadPositionData.X, QuadPositionData.Y, QuadPositionData.Z, QuadPositionData.W);
-				RectParams.SourceTextureSize = PostProcessTexture->GetSizeXY();
-				RectParams.CornerRadius = ShaderParams.PixelParams3;
-				RectParams.UITarget = Params.UITarget;
-				RectParams.HDRDisplayColorGamut = Params.HDRDisplayColorGamut;
-
-				RectParams.RestoreStateFunc = [&](FRHICommandListImmediate&InRHICmdList, FGraphicsPipelineStateInitializer& InGraphicsPSOInit, FRHIRenderPassInfo& RPInfo) {
-					return UpdateScissorRect(
-						InRHICmdList,
-#if STATS
-						ScissorClips,
-						StencilClips,
-#endif
-						RectParams.StencilRef,
-						MaskingID,
-						BackBuffer,
-						RenderBatch,
-						RPInfo,
-						DepthStencilTarget,
-						LastClippingState,
-						ViewTranslation2D,
-						InGraphicsPSOInit,
-						StencilVertexBuffer,
-						FMatrix(Params.ViewProjectionMatrix),
-						true);
-				};
-
-				RectParams.StencilRef = StencilRef;
-
-				FBlurRectParams BlurParams;
-				BlurParams.KernelSize = ShaderParams.PixelParams2.X;
-				BlurParams.Strength = ShaderParams.PixelParams2.Y;
-				BlurParams.DownsampleAmount = ShaderParams.PixelParams2.Z;
-
-				PostProcessor->BlurRect(RHICmdList, RendererModule, BlurParams, RectParams);
-
-				check(RHICmdList.IsOutsideRenderPass());
-				// Render pass for slate elements will be restarted on a next loop iteration if any
-			}
-		}
-		else
-		{
-			ICustomSlateElement* CustomDrawer = RenderBatch.CustomDrawer;
-			if (CustomDrawer)
-			{
-				// CustomDrawers will change the rendertarget. So we must close any outstanding renderpasses.
-				// Render pass for slate elements will be restarted on a next loop iteration if any
-				RHICmdList.EndRenderPass();	
-
-				SLATE_DRAW_EVENT(RHICmdList, CustomDrawer);
-
-				// Disable scissor rect. A previous draw element may have had one
-				RHICmdList.SetScissorRect(false, 0, 0, 0, 0);
-				LastClippingState = nullptr;
-
-				ICustomSlateElement::FSlateCustomDrawParams CustomDrawParams = ICustomSlateElement::FSlateCustomDrawParams();
-				CustomDrawParams.ViewProjectionMatrix = Params.ViewProjectionMatrix;
-				CustomDrawParams.ViewOffset = Params.ViewOffset;
-				CustomDrawParams.ViewRect = Params.ViewRect;
-				CustomDrawParams.HDRDisplayColorGamut = Params.HDRDisplayColorGamut;
-				CustomDrawParams.UsedSlatePostBuffers = Params.UsedSlatePostBuffers;
-				CustomDrawParams.bWireFrame = Params.bWireFrame;
-				CustomDrawParams.bIsHDR = Params.bIsHDR;
-
-				// This element is custom and has no Slate geometry.  Tell it to render itself now
-				if (CustomDrawer->UsesAdditionalRHIParams())
-				{
-					ICustomSlateElementRHI* CustomDrawerRHI = static_cast<ICustomSlateElementRHI*>(CustomDrawer);
-					CustomDrawerRHI->Draw_RHIRenderThread(RHICmdList, BackBuffer.GetRenderTargetTexture(), CustomDrawParams, FSlateRHIRenderingPolicyInterface(this));
-				}
-				else
-				{
-					CustomDrawer->Draw_RenderThread(RHICmdList, &BackBuffer.GetRenderTargetTexture(), CustomDrawParams);
-				}
-
-				//We reset the maskingID here because otherwise the RT might not get re-set in the lines above see: if (bClearStencil || bForceStateChange)
-				MaskingID = 0;
-			}
-		} // CustomDrawer
-	}
-
-	// Don't do color correction on iOS or Android, we don't have the GPU overhead for it.
-#if !(PLATFORM_IOS || PLATFORM_ANDROID)
-	if (bApplyColorDeficiencyCorrection && GSlateColorDeficiencyType != EColorVisionDeficiency::NormalVision && GSlateColorDeficiencySeverity > 0)
-	{
-		if (RHICmdList.IsInsideRenderPass())
-		{
-			RHICmdList.EndRenderPass();
-		}
-
-		FPostProcessRectParams RectParams;
-		RectParams.SourceTexture = BackBuffer.GetRenderTargetTexture();
-		RectParams.SourceRect = FSlateRect(0, 0, BackBuffer.GetSizeXY().X, BackBuffer.GetSizeXY().Y);
-		RectParams.DestRect = FSlateRect(0, 0, BackBuffer.GetSizeXY().X, BackBuffer.GetSizeXY().Y);
-		RectParams.SourceTextureSize = BackBuffer.GetSizeXY();
-		RectParams.CornerRadius = FVector4f(0.f, 0.f, 0.f, 0.f);
-
-		PostProcessor->ColorDeficiency(RHICmdList, RendererModule, RectParams);
-
-		FRHIRenderPassInfo RPInfo(ColorTarget, ERenderTargetActions::Load_Store);
-		RPInfo.DepthStencilRenderTarget.DepthStencilTarget = DepthStencilTarget;
-
-		// @todo refactor this.
-		// ColorDeficiency has self-contained renderpasses. To avoid starting an empty renderpass we do not
-		// restart the renderpass here.
-	}
-#endif
-
-	// Disable scissor rect we no longer need this.
-	RHICmdList.SetScissorRect(false, 0, 0, 0, 0);
-	// Disable depth/stencil testing once we're done also.
-	GraphicsPSOInit.DepthStencilState = TStaticDepthStencilState<false, CF_Always>::GetRHI();
-
-	for (int i = 0; i < NumScenes; i++)
-	{
-		// Don't need to delete SceneViews b/c the SceneViewFamily will delete it when it goes away.
-		delete SceneViewFamilyContexts[i];
-	}
-
-	SceneViews.Empty();
-	SceneViewFamilyContexts.Empty();
-
-	INC_DWORD_STAT_BY(STAT_SlateScissorClips, ScissorClips);
-	INC_DWORD_STAT_BY(STAT_SlateStencilClips, StencilClips);
-
-	// Disable scissor rect. 
-	// This fixes drawing on Metal when the last drawn element used a valid scissor rect
-	RHICmdList.SetScissorRect(false, 0, 0, 0, 0);
-}
-
-ETextureSamplerFilter FSlateRHIRenderingPolicy::GetSamplerFilter(const UTexture* Texture) const
-{
-	// Default to point filtering.
-	ETextureSamplerFilter Filter = ETextureSamplerFilter::Point;
-
-	switch (Texture->Filter)
-	{
-	case TF_Nearest: 
-		Filter = ETextureSamplerFilter::Point; 
-		break;
-	case TF_Bilinear:
-		Filter = ETextureSamplerFilter::Bilinear; 
-		break;
-	case TF_Trilinear: 
-		Filter = ETextureSamplerFilter::Trilinear; 
-		break;
-
-		// TF_Default
-	default:
-		// Use LOD group value to find proper filter setting.
-		if (Texture->LODGroup < TextureLODGroups.Num())
-		{
-			Filter = TextureLODGroups[Texture->LODGroup].Filter;
-		}
-	}
-
-	return Filter;
-}
-
-TShaderRef<FSlateElementPS> FSlateRHIRenderingPolicy::GetTexturePixelShader(FGlobalShaderMap* ShaderMap, ESlateShader ShaderType, ESlateDrawEffect DrawEffects, bool bUseTextureGrayscale, bool bIsVirtualTexture)
-{
-	TShaderRef<FSlateElementPS> PixelShader;
-
-#if WITH_SLATE_VISUALIZERS
-	if ( CVarShowSlateOverdraw.GetValueOnRenderThread() != 0 )
-	{
-		PixelShader = TShaderMapRef<FSlateDebugOverdrawPS>(ShaderMap);
-	}
-	else
-#endif
-	{
-	const bool bDrawDisabled = EnumHasAllFlags( DrawEffects, ESlateDrawEffect::DisabledEffect );
-	const bool bUseTextureAlpha = !EnumHasAllFlags( DrawEffects, ESlateDrawEffect::IgnoreTextureAlpha );
-
-	if ( bDrawDisabled )
-	{
-		switch ( ShaderType )
-		{
-		default:
-		case ESlateShader::Default:
-			if ( bUseTextureAlpha )
-			{
-				if ( bIsVirtualTexture )
-				{
-					if (bUseTextureGrayscale)
-					{
-						PixelShader = TShaderMapRef<TSlateElementPS<ESlateShader::Default, true, true, true, true> >(ShaderMap);
-					}
-					else
-					{
-						PixelShader = TShaderMapRef<TSlateElementPS<ESlateShader::Default, true, true, false, true> >(ShaderMap);
-					}
-				}
-				else
-				{
-					if (bUseTextureGrayscale)
-					{
-						PixelShader = TShaderMapRef<TSlateElementPS<ESlateShader::Default, true, true, true, false> >(ShaderMap);
+						PixelShader = TShaderMapRef<TSlateElementPS<ESlateShader::Default, true, true, true, false> >(ShaderMap);
 					}
 					else
 					{
@@ -1428,7 +267,7 @@ TShaderRef<FSlateElementPS> FSlateRHIRenderingPolicy::GetTexturePixelShader(FGlo
 			}
 			else
 			{
-				if ( bIsVirtualTexture )
+				if (bIsVirtualTexture)
 				{
 					if (bUseTextureGrayscale)
 					{
@@ -1484,13 +323,13 @@ TShaderRef<FSlateElementPS> FSlateRHIRenderingPolicy::GetTexturePixelShader(FGlo
 	}
 	else
 	{
-		switch ( ShaderType )
+		switch (ShaderType)
 		{
 		default:
 		case ESlateShader::Default:
-			if ( bUseTextureAlpha )
+			if (bUseTextureAlpha)
 			{
-				if ( bIsVirtualTexture )
+				if (bIsVirtualTexture)
 				{
 					if (bUseTextureGrayscale)
 					{
@@ -1515,7 +354,7 @@ TShaderRef<FSlateElementPS> FSlateRHIRenderingPolicy::GetTexturePixelShader(FGlo
 			}
 			else
 			{
-				if ( bIsVirtualTexture )
+				if (bIsVirtualTexture)
 				{
 					if (bUseTextureGrayscale)
 					{
@@ -1540,7 +379,7 @@ TShaderRef<FSlateElementPS> FSlateRHIRenderingPolicy::GetTexturePixelShader(FGlo
 			}
 			break;
 		case ESlateShader::Border:
-			if ( bUseTextureAlpha )
+			if (bUseTextureAlpha)
 			{
 				PixelShader = TShaderMapRef<TSlateElementPS<ESlateShader::Border, false, true> >(ShaderMap);
 			}
@@ -1571,12 +410,10 @@ TShaderRef<FSlateElementPS> FSlateRHIRenderingPolicy::GetTexturePixelShader(FGlo
 	}
 	}
 
-#undef PixelShaderLookupTable
-
 	return PixelShader;
 }
 
-void FSlateRHIRenderingPolicy::ChooseMaterialShaderTypes(ESlateShader ShaderType, bool bUseInstancing, FMaterialShaderTypes& OutShaderTypes)
+void ChooseMaterialShaderTypes(ESlateShader ShaderType, bool bUseInstancing, FMaterialShaderTypes& OutShaderTypes)
 {
 	switch (ShaderType)
 	{
@@ -1619,7 +456,7 @@ void FSlateRHIRenderingPolicy::ChooseMaterialShaderTypes(ESlateShader ShaderType
 	}
 }
 
-EPrimitiveType FSlateRHIRenderingPolicy::GetRHIPrimitiveType(ESlateDrawPrimitive SlateType)
+EPrimitiveType GetRHIPrimitiveType(ESlateDrawPrimitive SlateType)
 {
 	switch(SlateType)
 	{
@@ -1629,89 +466,1086 @@ EPrimitiveType FSlateRHIRenderingPolicy::GetRHIPrimitiveType(ESlateDrawPrimitive
 	default:
 		return PT_TriangleList;
 	}
-
 };
 
+//////////////////////////////////////////////////////////////////////////
 
-void FSlateRHIRenderingPolicy::AddSceneAt(FSceneInterface* Scene, int32 Index)
+FSlateElementsBuffers BuildSlateElementsBuffers(FRDGBuilder& GraphBuilder, FSlateBatchData& BatchData)
 {
-	ResourceManager->AddSceneAt(Scene, Index);
-}
+	FSlateElementsBuffers ElementsBuffers;
 
-void FSlateRHIRenderingPolicy::ClearScenes()
-{
-	ResourceManager->ClearScenes();
-}
+	if (BatchData.GetRenderBatches().IsEmpty())
+	{
+		return ElementsBuffers;
+	}
 
-void FSlateRHIRenderingPolicy::FlushGeneratedResources()
-{
-	PostProcessor->ReleaseRenderTargets();
-}
+	{
+		const FSlateVertexArray& Data = BatchData.GetFinalVertexData();
 
-void FSlateRHIRenderingPolicy::TickPostProcessResources()
-{
-	PostProcessor->TickPostProcessResources();
+		FRDGBufferDesc BufferDesc;
+		BufferDesc.Usage = EBufferUsageFlags::VertexBuffer | EBufferUsageFlags::Volatile;
+		BufferDesc.BytesPerElement = sizeof(Data[0]);
+		BufferDesc.NumElements = BatchData.GetMaxNumFinalVertices();
+
+		if (BufferDesc.NumElements > 0)
+		{
+			ElementsBuffers.VertexBuffer = GraphBuilder.CreateBuffer(BufferDesc, TEXT("SlateElementsVertexBuffer"));
+			GraphBuilder.QueueBufferUpload(ElementsBuffers.VertexBuffer, Data.GetData(), Data.Num() * BufferDesc.BytesPerElement, ERDGInitialDataFlags::NoCopy);
+		}
+	}
+
+	{
+		const FSlateIndexArray& Data = BatchData.GetFinalIndexData();
+
+		FRDGBufferDesc BufferDesc;
+		BufferDesc.Usage = EBufferUsageFlags::IndexBuffer | EBufferUsageFlags::Volatile;
+		BufferDesc.BytesPerElement = sizeof(Data[0]);
+		BufferDesc.NumElements = BatchData.GetMaxNumFinalIndices();
+
+		if (BufferDesc.NumElements > 0)
+		{
+			ElementsBuffers.IndexBuffer = GraphBuilder.CreateBuffer(BufferDesc, TEXT("SlateElementIndexBuffer"));
+			GraphBuilder.QueueBufferUpload(ElementsBuffers.IndexBuffer, Data.GetData(), Data.Num() * BufferDesc.BytesPerElement, ERDGInitialDataFlags::NoCopy);
+		}
+	}
+
+	SET_DWORD_STAT(STAT_SlateNumLayers, BatchData.GetNumLayers());
+	SET_DWORD_STAT(STAT_SlateNumBatches, BatchData.GetNumFinalBatches());
+	SET_DWORD_STAT(STAT_SlateVertexCount, BatchData.GetFinalVertexData().Num());
+
+	return ElementsBuffers;
 }
 
-void FSlateRHIRenderingPolicy::BlurRectExternal(FRHICommandListImmediate& RHICmdList, FRHITexture* BlurSrc, FRHITexture* BlurDst, FIntRect SrcRect, FIntRect DstRect, float BlurStrength) const
-{
-	SLATE_DRAW_EVENT(RHICmdList, PostProcess);
+//////////////////////////////////////////////////////////////////////////
 
-	FIntPoint BlurDstExtent = FIntPoint(DstRect.Width(), DstRect.Height());
+struct FSlateSceneViewAllocateInputs
+{
+	FIntPoint TextureExtent             = FIntPoint::ZeroValue;
+	FIntRect ViewRect;
+	FMatrix44f ViewProjectionMatrix     = FMatrix44f::Identity;
+	FIntPoint CursorPosition            = FIntPoint::ZeroValue;
+	FGameTime Time;
+	float ViewportScaleUI               = 1.0f;
+};
 
-	// If the radius isn't set, auto-compute it based on the strength
-	int32 OutKernelSize = FMath::RoundToInt(BlurStrength * 3.f);
+struct FSlateSceneView
+{
+	const FSceneInterface* Scene = nullptr;
+	ERHIFeatureLevel::Type FeatureLevel = ERHIFeatureLevel::Num;
+};
 
-	// Downsample if needed
-	int32 OutDownsampleAmount = 0;
-	if (OutKernelSize > 9)
+class FSlateSceneViewAllocator
+{
+	RDG_FRIEND_ALLOCATOR_FRIEND(FSlateSceneViewAllocator);
+public:
+	static FSlateSceneViewAllocator* Create(FRDGBuilder& GraphBuilder, FSlateRHIResourceManager& ResourceManager)
 	{
-		OutDownsampleAmount = OutKernelSize >= 64 ? 4 : 2;
-		OutKernelSize /= OutDownsampleAmount;
+		return GraphBuilder.AllocObject<FSlateSceneViewAllocator>(ResourceManager);
 	}
 
-	// Kernel sizes must be odd
-	if (OutKernelSize % 2 == 0)
+	const TUniformBufferRef<FViewUniformShaderParameters>& GetViewUniformBuffer(const FSlateSceneView* View) const
 	{
-		++OutKernelSize;
+		return UniformBuffers[(int32)View->FeatureLevel];
 	}
 
-	float ComputedStrength = FMath::Max(.5f, BlurStrength);
+	const FSlateSceneView* BeginAllocateSceneView(FRDGBuilder& GraphBuilder, FSlateRHIResourceManager& ResourceManager, int32 SceneViewIndex, const FSlateSceneViewAllocateInputs& Inputs)
+	{
+		if (!SceneViews.IsValidIndex(SceneViewIndex))
+		{
+			SceneViewIndex = SceneViewWithNullSceneIndex;
+		}
+
+		const FSlateSceneView& SceneView = SceneViews[SceneViewIndex];
+		const ERHIFeatureLevel::Type FeatureLevel = SceneView.FeatureLevel;
 
-	int32 RenderTargetWidth = BlurDstExtent.X;
-	int32 RenderTargetHeight = BlurDstExtent.Y;
-	
-	if (OutDownsampleAmount > 0)
+		if (TUniformBufferRef<FViewUniformShaderParameters>& UniformBuffer = UniformBuffers[(int32)FeatureLevel]; !UniformBuffer)
+		{
+			UniformBuffer = CreateUniformBuffer(FeatureLevel, Inputs);
+		}
+
+		return &SceneView;
+	}
+
+private:
+	FSlateSceneViewAllocator(FSlateRHIResourceManager& ResourceManager)
+		: SceneViewWithNullSceneIndex(ResourceManager.GetSceneCount())
+		, NumScenes(SceneViewWithNullSceneIndex + 1)
 	{
-		RenderTargetWidth = FMath::DivideAndRoundUp(RenderTargetWidth, OutDownsampleAmount);
-		RenderTargetHeight = FMath::DivideAndRoundUp(RenderTargetHeight, OutDownsampleAmount);
-		ComputedStrength /= OutDownsampleAmount;
+		SceneViews.SetNum(NumScenes);
+		SceneViews.Last().FeatureLevel = GMaxRHIFeatureLevel;
+
+		for (int32 Index = 0; Index < SceneViewWithNullSceneIndex; ++Index)
+		{
+			const FSceneInterface* Scene = ResourceManager.GetSceneAt(Index);
+			SceneViews[Index].Scene = Scene;
+			SceneViews[Index].FeatureLevel = Scene->GetFeatureLevel();
+		}
 	}
 
-	OutKernelSize = FMath::Clamp(OutKernelSize, 3, 255 /*MaxKernelSize*/);
+	static TUniformBufferRef<FViewUniformShaderParameters> CreateUniformBuffer(ERHIFeatureLevel::Type FeatureLevel, const FSlateSceneViewAllocateInputs& Inputs)
+	{
+		TRACE_CPUPROFILER_EVENT_SCOPE(FSlateSceneViewAllocator::CreateUniformBuffer);
 
-	FVector4f PostProcessData = FVector4f((float)OutKernelSize, ComputedStrength, (float)RenderTargetWidth, (float)RenderTargetHeight);
+		static const FEngineShowFlags DefaultShowFlags(ESFIM_Game);
 
-	FVector2f TopLeft = FVector2f::ZeroVector;
-	FVector2f BotRight = FVector2f(BlurDstExtent.X, BlurDstExtent.Y);
+		FIntRect ViewRect = Inputs.ViewRect;
 
-	FPostProcessRectParams RectParams;
-	RectParams.SourceTexture = BlurSrc;
-	RectParams.SourceRect = FSlateRect((float)SrcRect.Min.X, (float)SrcRect.Min.Y, (float)SrcRect.Max.X, (float)SrcRect.Max.Y);
-	RectParams.DestRect = FSlateRect(TopLeft.X, TopLeft.Y, BotRight.X, BotRight.Y);
-	RectParams.SourceTextureSize = BlurSrc->GetSizeXY();
-	RectParams.CornerRadius = FVector4f(0, 0, 0, 0);
-	RectParams.DestTexture = BlurDst;
-	RectParams.PostProcessDest = EPostProcessDestination::DestTexture;
+		// The window we are rendering to might not have a viewport, so use the full output instead.
+		if (ViewRect.IsEmpty())
+		{
+			ViewRect.Max = Inputs.TextureExtent;
+		}
 
-	FBlurRectParams BlurParams;
-	BlurParams.KernelSize = PostProcessData.X;
-	BlurParams.Strength = PostProcessData.Y;
-	BlurParams.DownsampleAmount = OutDownsampleAmount;
+		FViewMatrices::FMinimalInitializer Initializer;
+		Initializer.ProjectionMatrix = FMatrix(Inputs.ViewProjectionMatrix);
+		Initializer.ConstrainedViewRect = ViewRect;
 
-	IRendererModule& RendererModule = FModuleManager::GetModuleChecked<IRendererModule>(RendererModuleName);
-	PostProcessor->BlurRect(RHICmdList, RendererModule, BlurParams, RectParams);
+		const FViewMatrices ViewMatrices(Initializer);
 
-	RHICmdList.Transition(FRHITransitionInfo(BlurDst, ERHIAccess::RTV, ERHIAccess::SRVGraphics));
-}
+		const FSetupViewUniformParametersInputs SetupViewUniformParameterInputs =
+		{
+			  .EngineShowFlags  = &DefaultShowFlags
+			, .UnscaledViewRect = ViewRect
+			, .Time             = Inputs.Time
+			, .CursorPosition   = Inputs.CursorPosition
+		};
+
+		FViewUniformShaderParameters ViewUniformShaderParameters;
+		ViewUniformShaderParameters.VTFeedbackBuffer = GEmptyStructuredBufferWithUAV->UnorderedAccessViewRHI;
+
+		SetupCommonViewUniformBufferParameters(ViewUniformShaderParameters, Inputs.TextureExtent, 1, ViewRect, ViewMatrices, ViewMatrices, SetupViewUniformParameterInputs);
+
+		// Update Viewport Scale UI from any external sources (Material editor, UMG zoom scale / etc).
+		ViewUniformShaderParameters.ViewportScaleUI = Inputs.ViewportScaleUI;
+
+		// Always Update cursor position in realtime for slate
+		ViewUniformShaderParameters.CursorPosition = Inputs.CursorPosition;
+
+		// Slate materials need this scale to be positive, otherwise it can fail in querying scene textures (e.g., custom stencil)
+		ViewUniformShaderParameters.BufferToSceneTextureScale = FVector2f(1.0f, 1.0f);
+
+		ViewUniformShaderParameters.MobilePreviewMode = (FeatureLevel == ERHIFeatureLevel::ES3_1) && GMaxRHIFeatureLevel > ERHIFeatureLevel::ES3_1 ? 1.0f : 0.0f;
+
+		UpdateNoiseTextureParameters(ViewUniformShaderParameters);
+		return TUniformBufferRef<FViewUniformShaderParameters>::CreateUniformBufferImmediate(ViewUniformShaderParameters, UniformBuffer_SingleFrame);
+	}
+
+	const int32 SceneViewWithNullSceneIndex;
+	const int32 NumScenes;
+	TArray<FSlateSceneView, FRDGArrayAllocator> SceneViews;
+	TStaticArray<TUniformBufferRef<FViewUniformShaderParameters>, (int32)ERHIFeatureLevel::Num> UniformBuffers;
+};
+
+//////////////////////////////////////////////////////////////////////////
+
+bool GetSlateClippingPipelineState(const FSlateClippingOp* ClippingStateOp, FRHIDepthStencilState*& OutDepthStencilState, uint8& OutStencilRef)
+{
+	if (ClippingStateOp && ClippingStateOp->Method == EClippingMethod::Stencil)
+	{
+		// Setup the stenciling state to be read only now, disable depth writes, and restore the color buffer
+		// because we're about to go back to rendering widgets "normally", but with the added effect that now
+		// we have the stencil buffer bound with a bunch of clipping zones rendered into it.
+		OutDepthStencilState = 
+			TStaticDepthStencilState<
+				/*bEnableDepthWrite*/ false
+			, /*DepthTest*/ CF_Always
+			, /*bEnableFrontFaceStencil*/ true
+			, /*FrontFaceStencilTest*/ CF_Equal
+			, /*FrontFaceStencilFailStencilOp*/ SO_Keep
+			, /*FrontFaceDepthFailStencilOp*/ SO_Keep
+			, /*FrontFacePassStencilOp*/ SO_Keep
+			, /*bEnableBackFaceStencil*/ true
+			, /*BackFaceStencilTest*/ CF_Equal
+			, /*BackFaceStencilFailStencilOp*/ SO_Keep
+			, /*BackFaceDepthFailStencilOp*/ SO_Keep
+			, /*BackFacePassStencilOp*/ SO_Keep
+			, /*StencilReadMask*/ 0xFF
+			, /*StencilWriteMask*/ 0xFF>::GetRHI();
+
+		// Set a StencilRef equal to the number of stenciling/clipping masks, so unless the pixel we're rendering
+		// to is on top of a stencil pixel with the same number it's going to get rejected, thereby clipping
+		// everything except for the cross-section of all the stenciling quads.
+		OutStencilRef = ClippingStateOp->MaskingId + ClippingStateOp->Data_Stencil.Zones.Num();
+		return true;
+	}
+	else
+	{
+		OutDepthStencilState = TStaticDepthStencilState<false, CF_Always>::GetRHI();
+		OutStencilRef = 0;
+		return false;
+	}
+}
+
+void SetSlateClipping(FRHICommandList& RHICmdList, const FSlateClippingOp* ClippingStateOp, FIntRect ViewportRect)
+{
+	check(RHICmdList.IsInsideRenderPass());
+
+	if (ClippingStateOp)
+	{
+		const FVector2f ElementOffset = ClippingStateOp->Offset;
+
+		const auto ClampRectToViewport = [ElementOffset, ViewportRect] (FSlateRect ScissorRect)
+		{
+			ScissorRect.Left   = FMath::Clamp(ScissorRect.Left + ElementOffset.X, (float)ViewportRect.Min.X, (float)ViewportRect.Max.X);
+			ScissorRect.Top    = FMath::Clamp(ScissorRect.Top  + ElementOffset.Y, (float)ViewportRect.Min.Y, (float)ViewportRect.Max.Y);
+			ScissorRect.Right  = FMath::Clamp(ScissorRect.Right,  ScissorRect.Left, (float)ViewportRect.Max.X);
+			ScissorRect.Bottom = FMath::Clamp(ScissorRect.Bottom, ScissorRect.Top,  (float)ViewportRect.Max.Y);
+			return ScissorRect;
+		};
+
+		if (ClippingStateOp->Method == EClippingMethod::Scissor)
+		{
+			const FSlateRect ScissorRect = ClampRectToViewport(ClippingStateOp->Data_Scissor.Rect);
+			RHICmdList.SetScissorRect(true, ScissorRect.Left, ScissorRect.Top, ScissorRect.Right, ScissorRect.Bottom);
+		}
+		else
+		{
+			check(ClippingStateOp->Method == EClippingMethod::Stencil);
+
+			const TConstArrayView<FSlateClippingZone> Zones = ClippingStateOp->Data_Stencil.Zones;
+			check(Zones.Num() > 0);
+
+			// There might be some large - useless stencils, especially in the first couple of stencils if large
+			// widgets that clip also contain render targets, so, by setting the scissor to the AABB of the final
+			// stencil, we can cut out a lot of work that can't possibly be useful. We also round it, because if we
+			// don't it can over-eagerly slice off pixels it shouldn't.
+
+			const FSlateRect ScissorRect = ClampRectToViewport(Zones.Last().GetBoundingBox().Round());
+			RHICmdList.SetScissorRect(true, ScissorRect.Left, ScissorRect.Top, ScissorRect.Right, ScissorRect.Bottom);
+
+			const uint8 MaskingId = ClippingStateOp->MaskingId;
+
+			FGlobalShaderMap* ShaderMap = GetGlobalShaderMap(GMaxRHIShaderPlatform);
+			TShaderMapRef<FSlateMaskingVS> VertexShader(ShaderMap);
+			TShaderMapRef<FSlateMaskingPS> PixelShader(ShaderMap);
+
+			// Start by setting up the stenciling states so that we can write representations of the clipping zones into the stencil buffer only.
+			FGraphicsPipelineStateInitializer WriteMaskPSOInit;
+			RHICmdList.ApplyCachedRenderTargets(WriteMaskPSOInit);
+			WriteMaskPSOInit.BlendState = TStaticBlendStateWriteMask<CW_NONE, CW_NONE, CW_NONE, CW_NONE, CW_NONE, CW_NONE, CW_NONE, CW_NONE>::GetRHI();
+			WriteMaskPSOInit.RasterizerState = TStaticRasterizerState<>::GetRHI();
+			WriteMaskPSOInit.DepthStencilState =
+				TStaticDepthStencilState<
+				  /*bEnableDepthWrite*/ false
+				, /*DepthTest*/ CF_Always
+				, /*bEnableFrontFaceStencil*/ true
+				, /*FrontFaceStencilTest*/ CF_Always
+				, /*FrontFaceStencilFailStencilOp*/ SO_Keep
+				, /*FrontFaceDepthFailStencilOp*/ SO_Keep
+				, /*FrontFacePassStencilOp*/ SO_Replace
+				, /*bEnableBackFaceStencil*/ true
+				, /*BackFaceStencilTest*/ CF_Always
+				, /*BackFaceStencilFailStencilOp*/ SO_Keep
+				, /*BackFaceDepthFailStencilOp*/ SO_Keep
+				, /*BackFacePassStencilOp*/ SO_Replace
+				, /*StencilReadMask*/ 0xFF
+				, /*StencilWriteMask*/ 0xFF>::GetRHI();
+
+			WriteMaskPSOInit.BoundShaderState.VertexDeclarationRHI = GSlateMaskingVertexDeclaration.VertexDeclarationRHI;
+			WriteMaskPSOInit.BoundShaderState.VertexShaderRHI = VertexShader.GetVertexShader();
+			WriteMaskPSOInit.BoundShaderState.PixelShaderRHI = PixelShader.GetPixelShader();
+			WriteMaskPSOInit.PrimitiveType = PT_TriangleStrip;
+
+			// Draw the first stencil using SO_Replace, so that we stomp any pixel with a MaskingID + 1.
+			SetGraphicsPipelineState(RHICmdList, WriteMaskPSOInit, MaskingId + 1);
+
+			// Simple 2D orthographic projection from screen space to NDC space.
+			const FVector2f A
+			(
+				2.0f /  ViewportRect.Width(),
+				2.0f / -ViewportRect.Height()
+			);
+
+			const FVector2f B
+			(
+				(ViewportRect.Min.X + ViewportRect.Max.X) / -ViewportRect.Width(),
+				(ViewportRect.Min.Y + ViewportRect.Max.Y) /  ViewportRect.Height()
+			);
+
+			const auto TransformVertex = [A, B, ElementOffset](FVector2f P)
+			{
+				return FVector2f((P.X + ElementOffset.X) * A.X + B.X, (P.Y + ElementOffset.Y) * A.Y + B.Y);
+			};
+
+			const auto SetMaskingParameters = [VertexShader, TransformVertex] (FRHIBatchedShaderParameters& BatchedParameters, const FSlateClippingZone& Zone)
+			{
+				FSlateMaskingVS::FParameters Parameters;
+				Parameters.MaskRectPacked[0] = FVector4f(TransformVertex(Zone.TopLeft), TransformVertex(Zone.TopRight));
+				Parameters.MaskRectPacked[1] = FVector4f(TransformVertex(Zone.BottomLeft), TransformVertex(Zone.BottomRight));
+				SetShaderParameters(BatchedParameters, VertexShader, Parameters);
+			};
+
+			{
+				FRHIBatchedShaderParameters& BatchedParameters = RHICmdList.GetScratchShaderParameters();
+				SetMaskingParameters(BatchedParameters, Zones[0]);
+				RHICmdList.SetBatchedShaderParameters(VertexShader.GetVertexShader(), BatchedParameters);
+				RHICmdList.SetStreamSource(0, GSlateStencilClipVertexBuffer.VertexBufferRHI, 0);
+				RHICmdList.DrawPrimitive(0, 2, 1);
+			}
+
+			// Now setup the pipeline to use SO_SaturatedIncrement, since we've established the initial
+			// stencil with SO_Replace, we can safely use SO_SaturatedIncrement, to build up the stencil
+			// to the required mask of MaskingID + StencilQuads.Num(), thereby ensuring only the union of
+			// all stencils will render pixels.
+			WriteMaskPSOInit.DepthStencilState =
+				TStaticDepthStencilState<
+				/*bEnableDepthWrite*/ false
+				, /*DepthTest*/ CF_Always
+				, /*bEnableFrontFaceStencil*/ true
+				, /*FrontFaceStencilTest*/ CF_Always
+				, /*FrontFaceStencilFailStencilOp*/ SO_Keep
+				, /*FrontFaceDepthFailStencilOp*/ SO_Keep
+				, /*FrontFacePassStencilOp*/ SO_SaturatedIncrement
+				, /*bEnableBackFaceStencil*/ true
+				, /*BackFaceStencilTest*/ CF_Always
+				, /*BackFaceStencilFailStencilOp*/ SO_Keep
+				, /*BackFaceDepthFailStencilOp*/ SO_Keep
+				, /*BackFacePassStencilOp*/ SO_SaturatedIncrement
+				, /*StencilReadMask*/ 0xFF
+				, /*StencilWriteMask*/ 0xFF>::GetRHI();
+
+			SetGraphicsPipelineState(RHICmdList, WriteMaskPSOInit, 0);
+
+			// Next write the number of quads representing the number of clipping zones have on top of each other.
+			for (int32 MaskIndex = 1; MaskIndex < Zones.Num(); MaskIndex++)
+			{
+				FRHIBatchedShaderParameters& BatchedParameters = RHICmdList.GetScratchShaderParameters();
+				SetMaskingParameters(BatchedParameters, Zones[MaskIndex]);
+				RHICmdList.SetBatchedShaderParameters(VertexShader.GetVertexShader(), BatchedParameters);
+				RHICmdList.SetStreamSource(0, GSlateStencilClipVertexBuffer.VertexBufferRHI, 0);
+				RHICmdList.DrawPrimitive(0, 2, 1);
+			}
+		}
+	}
+	else
+	{
+		RHICmdList.SetScissorRect(false, 0.0f, 0.0f, 0.0f, 0.0f);
+	}
+}
+
+enum class ESlateClippingStencilAction : uint8
+{
+	None,
+	Write,
+	Clear
+};
+
+struct FSlateClippingCreateContext
+{
+	uint32 NumStencils = 0;
+	uint32 NumScissors = 0;
+	uint32 MaskingId = 0;
+	ESlateClippingStencilAction StencilAction = ESlateClippingStencilAction::None;
+};
+
+const FSlateClippingOp* CreateSlateClipping(FRDGBuilder& GraphBuilder, const FVector2f ElementsOffset, const FSlateClippingState* ClippingState, FSlateClippingCreateContext& Context)
+{
+	Context.StencilAction = ESlateClippingStencilAction::None;
+
+	if (ClippingState)
+	{
+		if (ClippingState->GetClippingMethod() == EClippingMethod::Scissor)
+		{
+			Context.NumScissors++;
+
+			const FSlateClippingZone& ScissorRect = ClippingState->ScissorRect.GetValue();
+
+			return FSlateClippingOp::Scissor(GraphBuilder, ElementsOffset, FSlateRect(ScissorRect.TopLeft.X, ScissorRect.TopLeft.Y, ScissorRect.BottomRight.X, ScissorRect.BottomRight.Y));
+		}
+		else
+		{
+			Context.NumStencils++;
+
+			TConstArrayView<FSlateClippingZone> StencilQuads = ClippingState->StencilQuads;
+			check(StencilQuads.Num() > 0);
+
+			// Reset the masking ID back to zero if stencil is going to overflow.
+			if (Context.MaskingId + StencilQuads.Num() > 255)
+			{
+				Context.MaskingId = 0;
+			}
+
+			// Mark stencil for clear when the masking id is 0.
+			Context.StencilAction = Context.MaskingId == 0 ? ESlateClippingStencilAction::Clear : ESlateClippingStencilAction::Write;
+
+			const FSlateClippingOp* Op = FSlateClippingOp::Stencil(GraphBuilder, ElementsOffset, StencilQuads, Context.MaskingId);
+			Context.MaskingId += StencilQuads.Num();
+			return Op;
+		}
+	}
+	return nullptr;
+}
+
+//////////////////////////////////////////////////////////////////////////
+
+enum class ESlateRenderBatchType
+{
+	CustomDrawer,
+	PostProcess,
+	Primitive,
+	MAX
+};
+
+inline ESlateRenderBatchType GetSlateRenderBatchType(const FSlateRenderBatch& DrawBatch)
+{
+	if (DrawBatch.CustomDrawer != nullptr)
+	{
+		return ESlateRenderBatchType::CustomDrawer;
+	}
+
+	if (DrawBatch.ShaderType == ESlateShader::PostProcess)
+	{
+		return ESlateRenderBatchType::PostProcess;
+	}
+
+	return ESlateRenderBatchType::Primitive;
+}
+
+struct FSlateRenderBatchOp
+{
+	FSlateRenderBatchOp() = default;
+
+	FSlateRenderBatchOp(const FSlateRenderBatch* InRenderBatch, const FSlateClippingOp* InClippingStateOp, const FSlateSceneView* InSceneView)
+		: RenderBatch(InRenderBatch)
+		, ClippingStateOp(InClippingStateOp)
+		, SceneView(InSceneView)
+	{}
+
+	const FSlateRenderBatch* RenderBatch;
+	const FSlateClippingOp* ClippingStateOp;
+	const FSlateSceneView* SceneView = nullptr;
+};
+
+//////////////////////////////////////////////////////////////////////////
+
+struct FSlateRenderBatchDrawInputs
+{
+	FMatrix44f ElementsMatrix;
+	FVector2f ElementsOffset;
+	FIntRect ElementsViewRect;
+	FIntRect SceneViewRect;
+	const FSlateSceneViewAllocator* SceneViewAllocator;
+	TConstArrayView<FTextureLODGroup> TextureLODGroups;
+	FSlateElementsBuffers ElementsBuffers;
+	float DisplayGamma;
+	float DisplayContrast;
+	float EngineGamma;
+	bool bWireframe;
+#if WITH_SLATE_VISUALIZERS
+	FLinearColor BatchColor;
+#endif
+};
+
+struct FSlateRenderBatchDrawState
+{
+	const FSlateClippingOp* LastClippingOp = nullptr;
+	FGraphicsPipelineStateInitializer GraphicsPSOInit;
+	uint8 StencilRef = 0;
+};
+
+void DrawSlateRenderBatch(
+	FRHICommandList& RHICmdList,
+	FSlateRenderBatchDrawState& State,
+	const FSlateRenderBatchDrawInputs& Inputs,
+	const FSlateRenderBatchOp& RenderBatchOp)
+{
+	const FSlateClippingOp* ClippingStateOp = RenderBatchOp.ClippingStateOp;
+	const FSlateRenderBatch& RenderBatch         = *RenderBatchOp.RenderBatch;
+	
+	FGlobalShaderMap* ShaderMap = GetGlobalShaderMap(GMaxRHIShaderPlatform);
+
+	if (State.LastClippingOp != ClippingStateOp)
+	{
+		GetSlateClippingPipelineState(ClippingStateOp, State.GraphicsPSOInit.DepthStencilState, State.StencilRef);
+		SetSlateClipping(RHICmdList, ClippingStateOp, Inputs.ElementsViewRect);
+		State.LastClippingOp = ClippingStateOp;
+	}
+
+	FRHIBuffer* ElementsVertexBuffer = Inputs.ElementsBuffers.VertexBuffer->GetRHI();
+	FRHIBuffer* ElementsIndexBuffer  = Inputs.ElementsBuffers.IndexBuffer->GetRHI();
+
+	const FSlateShaderResource* ShaderResource   = RenderBatch.ShaderResource;
+	const ESlateBatchDrawFlag DrawFlags          = RenderBatch.DrawFlags;
+	const ESlateDrawEffect DrawEffects           = RenderBatch.DrawEffects;
+	const ESlateShader ShaderType                = RenderBatch.ShaderType;
+	const FShaderParams& ShaderParams            = RenderBatch.ShaderParams;
+
+	if (EnumHasAllFlags(DrawFlags, ESlateBatchDrawFlag::Wireframe))
+	{
+		State.GraphicsPSOInit.RasterizerState = TStaticRasterizerState<FM_Wireframe>::GetRHI();
+	}
+	else
+	{
+		State.GraphicsPSOInit.RasterizerState = TStaticRasterizerState<FM_Solid>::GetRHI();
+	}
+
+	check(RenderBatch.NumIndices > 0);
+	const uint32 PrimitiveCount = RenderBatch.DrawPrimitiveType == ESlateDrawPrimitive::LineList ? RenderBatch.NumIndices / 2 : RenderBatch.NumIndices / 3;
+
+	check(ShaderResource == nullptr || !ShaderResource->Debug_IsDestroyed());
+	const ESlateShaderResource::Type ResourceType = ShaderResource ? ShaderResource->GetType() : ESlateShaderResource::Invalid;
+
+	const bool bUseInstancing = RenderBatch.InstanceCount > 0 && RenderBatch.InstanceData != nullptr;
+
+	if (ResourceType == ESlateShaderResource::Material)
+	{
+		FSlateMaterialResource* MaterialShaderResource = (FSlateMaterialResource*)ShaderResource;
+		MaterialShaderResource->CheckForStaleResources();
+
+		const FMaterialRenderProxy* MaterialRenderProxy = MaterialShaderResource->GetRenderProxy();
+		if (!MaterialRenderProxy)
+		{
+			return;
+		}
+		
+		const FSlateSceneView* SceneView = RenderBatchOp.SceneView;
+		const ERHIFeatureLevel::Type SceneFeatureLevel = SceneView->FeatureLevel;
+		const FSceneInterface* Scene = SceneView->Scene;
+		const TUniformBufferRef<FViewUniformShaderParameters>& ViewUniformBuffer = Inputs.SceneViewAllocator->GetViewUniformBuffer(SceneView);
+
+		MaterialRenderProxy->UpdateUniformExpressionCacheIfNeeded(RHICmdList, SceneFeatureLevel);
+
+		TShaderRef<FSlateMaterialShaderVS> VertexShader;
+		TShaderRef<FSlateMaterialShaderPS> PixelShader;
+
+		FMaterialShaderTypes ShaderTypesToGet;
+		ChooseMaterialShaderTypes(ShaderType, bUseInstancing, ShaderTypesToGet);
+		const FMaterial* EffectiveMaterial = nullptr;
+
+		while (MaterialRenderProxy)
+		{
+			const FMaterial* Material = MaterialRenderProxy->GetMaterialNoFallback(SceneFeatureLevel);
+			FMaterialShaders Shaders;
+			if (Material && Material->TryGetShaders(ShaderTypesToGet, nullptr, Shaders))
+			{
+				EffectiveMaterial = Material;
+				Shaders.TryGetVertexShader(VertexShader);
+				Shaders.TryGetPixelShader(PixelShader);
+				break;
+			}
+
+			MaterialRenderProxy = MaterialRenderProxy->GetFallback(SceneFeatureLevel);
+		}
+
+		if (!VertexShader.IsValid() || !PixelShader.IsValid())
+		{
+			return;
+		}
+
+		SLATE_DRAW_EVENTF(RHICmdList, MaterialBatch, TEXT("Slate Material: %s"), MaterialRenderProxy->GetMaterialName());
+
+		PixelShader->SetBlendState(State.GraphicsPSOInit, EffectiveMaterial);
+
+		FSlateShaderResource* MaskResource = MaterialShaderResource->GetTextureMaskResource();
+
+		if (MaskResource && IsOpaqueOrMaskedBlendMode(*EffectiveMaterial))
+		{
+			// Font materials require some form of translucent blending
+			State.GraphicsPSOInit.BlendState = TStaticBlendState<CW_RGBA, BO_Add, BF_SourceAlpha, BF_InverseSourceAlpha, BO_Add, BF_InverseDestAlpha, BF_One>::GetRHI();
+		}
+
+		State.GraphicsPSOInit.BoundShaderState.VertexDeclarationRHI = bUseInstancing ? GSlateInstancedVertexDeclaration.VertexDeclarationRHI : GSlateVertexDeclaration.VertexDeclarationRHI;
+		State.GraphicsPSOInit.BoundShaderState.VertexShaderRHI = VertexShader.GetVertexShader();
+		State.GraphicsPSOInit.BoundShaderState.PixelShaderRHI = PixelShader.GetPixelShader();
+		State.GraphicsPSOInit.PrimitiveType = GetRHIPrimitiveType(RenderBatch.DrawPrimitiveType);
+
+		SetGraphicsPipelineState(RHICmdList, State.GraphicsPSOInit, State.StencilRef);
+
+		{
+			FRHIBatchedShaderParameters& BatchedParameters = RHICmdList.GetScratchShaderParameters();
+
+			VertexShader->SetViewProjection(BatchedParameters, Inputs.ElementsMatrix);
+			VertexShader->SetMaterialShaderParameters(BatchedParameters, Scene, ViewUniformBuffer, MaterialRenderProxy, EffectiveMaterial);
+
+			RHICmdList.SetBatchedShaderParameters(VertexShader.GetVertexShader(), BatchedParameters);
+		}
+
+		{
+			FRHIBatchedShaderParameters& BatchedParameters = RHICmdList.GetScratchShaderParameters();
+
+			PixelShader->SetParameters(BatchedParameters, Scene, ViewUniformBuffer, MaterialRenderProxy, EffectiveMaterial, ShaderParams);
+			const float FinalGamma = EnumHasAnyFlags(DrawFlags, ESlateBatchDrawFlag::ReverseGamma) ? 1.0f / Inputs.EngineGamma : EnumHasAnyFlags(DrawFlags, ESlateBatchDrawFlag::NoGamma) ? 1.0f : Inputs.DisplayGamma;
+			const float FinalContrast = EnumHasAnyFlags(DrawFlags, ESlateBatchDrawFlag::NoGamma) ? 1 : Inputs.DisplayContrast;
+			PixelShader->SetDisplayGammaAndContrast(BatchedParameters, FinalGamma, FinalContrast);
+			const bool bDrawDisabled = EnumHasAllFlags(DrawEffects, ESlateDrawEffect::DisabledEffect);
+			PixelShader->SetDrawFlags(BatchedParameters, bDrawDisabled);
+
+			if (MaskResource)
+			{
+				FTextureRHIRef TextureRHI = ((TSlateTexture<FTextureRHIRef>*)MaskResource)->GetTypedResource();
+
+				PixelShader->SetAdditionalTexture(BatchedParameters, TextureRHI, TStaticSamplerState<SF_Bilinear, AM_Clamp, AM_Clamp, AM_Clamp>::GetRHI());
+			}
+
+			RHICmdList.SetBatchedShaderParameters(PixelShader.GetPixelShader(), BatchedParameters);
+		}
+
+		if (bUseInstancing)
+		{
+			RenderBatch.InstanceData->BindStreamSource(RHICmdList, 1, RenderBatch.InstanceOffset);
+
+			RHICmdList.SetStreamSource(0, ElementsVertexBuffer, RenderBatch.VertexOffset * sizeof(FSlateVertex));
+			RHICmdList.DrawIndexedPrimitive(ElementsIndexBuffer, 0, 0, RenderBatch.NumVertices, RenderBatch.IndexOffset, PrimitiveCount, RenderBatch.InstanceCount);
+		}
+		else
+		{
+			RHICmdList.SetStreamSource(0, ElementsVertexBuffer, RenderBatch.VertexOffset * sizeof(FSlateVertex));
+			RHICmdList.SetStreamSource(1, nullptr, 0);
+			RHICmdList.DrawIndexedPrimitive(ElementsIndexBuffer, 0, 0, RenderBatch.NumVertices, RenderBatch.IndexOffset, PrimitiveCount, 1);
+		}
+	}
+	else
+	{
+		check(!bUseInstancing);
+
+		TShaderRef<FSlateElementPS> PixelShader;
+
+#if WITH_SLATE_VISUALIZERS
+		TShaderRef<FSlateDebugBatchingPS> BatchingPixelShader;
+		if (CVarShowSlateBatching.GetValueOnRenderThread() != 0)
+		{
+			BatchingPixelShader = TShaderMapRef<FSlateDebugBatchingPS>(ShaderMap);
+			PixelShader = BatchingPixelShader;
+		}
+		else
+#endif
+		{
+			bool bIsVirtualTexture = false;
+
+			// check if texture is using BC4 compression and set shader to render grayscale
+			bool bUseTextureGrayscale = false;
+
+			if (ShaderResource != nullptr && ResourceType == ESlateShaderResource::TextureObject)
+			{
+				FSlateBaseUTextureResource* TextureObjectResource = const_cast<FSlateBaseUTextureResource*>(static_cast<const FSlateBaseUTextureResource*>(ShaderResource));
+
+				if (UTexture* TextureObj = TextureObjectResource->GetTextureObject())
+				{
+					bIsVirtualTexture = TextureObj->IsCurrentlyVirtualTextured();
+
+					if (TextureObj->CompressionSettings == TC_Alpha)
+					{
+						bUseTextureGrayscale = true;
+					}
+				}
+			}
+
+			PixelShader = GetTexturePixelShader(ShaderMap, ShaderType, DrawEffects, bUseTextureGrayscale, bIsVirtualTexture);
+		}
+
+#if WITH_SLATE_VISUALIZERS
+		if (BatchingPixelShader.IsValid())
+		{
+			State.GraphicsPSOInit.BlendState = TStaticBlendState<CW_RGBA, BO_Add, BF_SourceAlpha, BF_InverseSourceAlpha, BO_Add, BF_One, BF_InverseSourceAlpha>::GetRHI();
+		}
+		else if (CVarShowSlateOverdraw.GetValueOnRenderThread() != 0)
+		{
+			State.GraphicsPSOInit.BlendState = TStaticBlendState<CW_RGB, BO_Add, BF_One, BF_One, BO_Add, BF_Zero, BF_InverseSourceAlpha>::GetRHI();
+		}
+		else
+#endif
+		{
+			State.GraphicsPSOInit.BlendState =
+				EnumHasAllFlags(DrawFlags, ESlateBatchDrawFlag::NoBlending)
+				? TStaticBlendState<>::GetRHI()
+				: (EnumHasAllFlags(DrawFlags, ESlateBatchDrawFlag::PreMultipliedAlpha)
+					? TStaticBlendState<CW_RGBA, BO_Add, BF_One, BF_InverseSourceAlpha, BO_Add, BF_One, BF_InverseSourceAlpha>::GetRHI()
+					: TStaticBlendState<CW_RGBA, BO_Add, BF_SourceAlpha, BF_InverseSourceAlpha, BO_Add, BF_One, BF_InverseSourceAlpha>::GetRHI())
+				;
+		}
+
+		if (EnumHasAllFlags(DrawFlags, ESlateBatchDrawFlag::Wireframe) || Inputs.bWireframe)
+		{
+			State.GraphicsPSOInit.RasterizerState = TStaticRasterizerState<FM_Wireframe>::GetRHI();
+
+			if (Inputs.bWireframe)
+			{
+				State.GraphicsPSOInit.BlendState = TStaticBlendState<>::GetRHI();
+			}
+		}
+		else
+		{
+			State.GraphicsPSOInit.RasterizerState = TStaticRasterizerState<FM_Solid>::GetRHI();
+		}
+
+		TShaderMapRef<FSlateElementVS> GlobalVertexShader(ShaderMap);
+
+		State.GraphicsPSOInit.BoundShaderState.VertexDeclarationRHI = GSlateVertexDeclaration.VertexDeclarationRHI;
+		State.GraphicsPSOInit.BoundShaderState.VertexShaderRHI = GlobalVertexShader.GetVertexShader();
+		State.GraphicsPSOInit.BoundShaderState.PixelShaderRHI = PixelShader.GetPixelShader();
+		State.GraphicsPSOInit.PrimitiveType = GetRHIPrimitiveType(RenderBatch.DrawPrimitiveType);
+
+		SetGraphicsPipelineState(RHICmdList, State.GraphicsPSOInit, State.StencilRef);
+
+		FRHISamplerState* SamplerState = TStaticSamplerState<SF_Bilinear, AM_Clamp, AM_Clamp, AM_Clamp>::GetRHI();
+		FRHITexture* TextureRHI = GWhiteTexture->TextureRHI;
+		bool bIsVirtualTexture = false;
+		FTextureResource* TextureResource = nullptr;
+
+		if (ShaderResource)
+		{
+			ETextureSamplerFilter Filter = ETextureSamplerFilter::Bilinear;
+
+			if (ResourceType == ESlateShaderResource::TextureObject)
+			{
+				FSlateBaseUTextureResource* TextureObjectResource = (FSlateBaseUTextureResource*)ShaderResource;
+				if (UTexture* TextureObj = TextureObjectResource->GetTextureObject())
+				{
+					TextureObjectResource->CheckForStaleResources();
+
+					TextureRHI = TextureObjectResource->AccessRHIResource();
+
+					// This can upset some RHIs, so use transparent black texture until it's valid.
+					// these can be temporarily invalid when recreating them / invalidating their streaming
+					// state.
+					if (TextureRHI == nullptr)
+					{
+						// We use transparent black here, because it's about to become valid - probably, and flashing white
+						// wouldn't be ideal.
+						TextureRHI = GTransparentBlackTexture->TextureRHI;
+					}
+
+					TextureResource = TextureObj->GetResource();
+
+					Filter = GetSamplerFilter(Inputs.TextureLODGroups, TextureObj);
+					bIsVirtualTexture = TextureObj->IsCurrentlyVirtualTextured();
+				}
+			}
+			else
+			{
+				FRHITexture* NativeTextureRHI = ((TSlateTexture<FTextureRHIRef>*)ShaderResource)->GetTypedResource();
+				// Atlas textures that have no content are never initialized but null textures are invalid on many platforms.
+				TextureRHI = NativeTextureRHI ? NativeTextureRHI : (FRHITexture*)GWhiteTexture->TextureRHI;
+			}
+
+			SamplerState = GetSamplerState(DrawFlags, Filter);
+		}
+
+		{
+			FRHIBatchedShaderParameters& BatchedParameters = RHICmdList.GetScratchShaderParameters();
+
+#if WITH_SLATE_VISUALIZERS
+			if (BatchingPixelShader.IsValid())
+			{
+				BatchingPixelShader->SetBatchColor(BatchedParameters, Inputs.BatchColor);
+			}
+#endif
+
+			if (bIsVirtualTexture && TextureResource != nullptr)
+			{
+				PixelShader->SetVirtualTextureParameters(BatchedParameters, static_cast<FVirtualTexture2DResource*>(TextureResource));
+			}
+			else
+			{
+				PixelShader->SetTexture(BatchedParameters, TextureRHI, SamplerState);
+			}
+
+			PixelShader->SetShaderParams(BatchedParameters, ShaderParams);
+			const float FinalGamma = EnumHasAnyFlags(DrawFlags, ESlateBatchDrawFlag::ReverseGamma) ? (1.0f / Inputs.EngineGamma) : EnumHasAnyFlags(DrawFlags, ESlateBatchDrawFlag::NoGamma) ? 1.0f : Inputs.DisplayGamma;
+			const float FinalContrast = EnumHasAnyFlags(DrawFlags, ESlateBatchDrawFlag::NoGamma) ? 1 : Inputs.DisplayContrast;
+			PixelShader->SetDisplayGammaAndInvertAlphaAndContrast(BatchedParameters, FinalGamma, EnumHasAllFlags(DrawEffects, ESlateDrawEffect::InvertAlpha) ? 1.0f : 0.0f, FinalContrast);
+
+			RHICmdList.SetBatchedShaderParameters(PixelShader.GetPixelShader(), BatchedParameters);
+		}
+
+		{
+			FRHIBatchedShaderParameters& BatchedParameters = RHICmdList.GetScratchShaderParameters();
+
+			GlobalVertexShader->SetViewProjection(BatchedParameters, Inputs.ElementsMatrix);
+
+			RHICmdList.SetBatchedShaderParameters(GlobalVertexShader.GetVertexShader(), BatchedParameters);
+		}
+
+		RHICmdList.SetStreamSource(0, ElementsVertexBuffer, RenderBatch.VertexOffset * sizeof(FSlateVertex));
+		RHICmdList.DrawIndexedPrimitive(ElementsIndexBuffer, 0, 0, RenderBatch.NumVertices, RenderBatch.IndexOffset, PrimitiveCount, RenderBatch.InstanceCount);
+	}
+}
+
+BEGIN_SHADER_PARAMETER_STRUCT(FSlateRenderBatchParameters, )
+	SHADER_PARAMETER_STRUCT_INCLUDE(FSceneTextureExtractsParameters, SceneTextures)
+	RDG_BUFFER_ACCESS(ElementsVertexBuffer, ERHIAccess::VertexOrIndexBuffer)
+	RDG_BUFFER_ACCESS(ElementsIndexBuffer, ERHIAccess::VertexOrIndexBuffer)
+	RENDER_TARGET_BINDING_SLOTS()
+END_SHADER_PARAMETER_STRUCT()
+
+void AddSlateDrawElementsPass(
+	FRDGBuilder& GraphBuilder,
+	const FSlateRHIRenderingPolicy& RenderingPolicy,
+	const FSlateDrawElementsPassInputs& Inputs,
+	TConstArrayView<FSlateRenderBatch> RenderBatches,
+	int32 FirstBatchIndex)
+{
+	TRACE_CPUPROFILER_EVENT_SCOPE(AddSlateDrawElements);
+	
+	const FIntPoint ElementsTextureExtent = Inputs.ElementsTexture->Desc.Extent;
+	const FScreenPassTexture ElementsTexture(Inputs.ElementsTexture);
+
+	const float EngineGamma  = GEngine ? GEngine->GetDisplayGamma() : 2.2f;
+	const float DisplayGamma = Inputs.bAllowGammaCorrection && !Inputs.bElementsTextureIsHDRDisplay ? EngineGamma : 1.0f;
+
+	FSlateRHIResourceManager& ResourceManager = RenderingPolicy.GetResourceManagerRHI();
+
+	FSlateSceneViewAllocator* SceneViewAllocator = FSlateSceneViewAllocator::Create(GraphBuilder, ResourceManager);
+
+#if WITH_SLATE_VISUALIZERS
+	FRandomStream BatchColors(1337);
+#endif
+
+	FSlateRenderBatchDrawInputs* RenderBatchDrawInputs = GraphBuilder.AllocPOD<FSlateRenderBatchDrawInputs>();
+
+	*RenderBatchDrawInputs =
+	{
+		  .ElementsMatrix        = Inputs.ElementsMatrix
+		, .ElementsOffset        = Inputs.ElementsOffset
+		, .ElementsViewRect      = ElementsTexture.ViewRect
+		, .SceneViewRect         = Inputs.SceneViewRect
+		, .SceneViewAllocator    = SceneViewAllocator
+		, .TextureLODGroups      = GetTextureLODGroups()
+		, .ElementsBuffers       = Inputs.ElementsBuffers
+		, .DisplayGamma          = DisplayGamma
+		, .DisplayContrast       = GSlateContrast
+		, .EngineGamma           = EngineGamma
+		, .bWireframe            = Inputs.bWireframe
+#if WITH_SLATE_VISUALIZERS
+		, .BatchColor            = FLinearColor(BatchColors.GetUnitVector())
+#endif
+	};
+
+	ERenderTargetLoadAction ElementsLoadAction = Inputs.ElementsLoadAction;
+
+	const auto ConsumeLoadAction = [] (ERenderTargetLoadAction& InOutLoadAction)
+	{
+		ERenderTargetLoadAction LoadAction = InOutLoadAction;
+		InOutLoadAction = ERenderTargetLoadAction::ELoad;
+		return LoadAction;
+	};
+
+	FSlateRenderBatchParameters* NoneStencilActionPassParameters = GraphBuilder.AllocParameters<FSlateRenderBatchParameters>();
+	NoneStencilActionPassParameters->SceneTextures = GetSceneTextureExtracts().GetShaderParameters();
+	NoneStencilActionPassParameters->ElementsVertexBuffer = Inputs.ElementsBuffers.VertexBuffer;
+	NoneStencilActionPassParameters->ElementsIndexBuffer  = Inputs.ElementsBuffers.IndexBuffer;
+	NoneStencilActionPassParameters->RenderTargets[0] = FRenderTargetBinding(Inputs.ElementsTexture, ERenderTargetLoadAction::ELoad);
+
+	FSlateRenderBatchParameters* ClearStencilActionPassParameters = nullptr;
+	FSlateRenderBatchParameters* WriteStencilActionPassParameters = nullptr;
+
+	if (Inputs.StencilTexture)
+	{
+		WriteStencilActionPassParameters = GraphBuilder.AllocParameters<FSlateRenderBatchParameters>(NoneStencilActionPassParameters);
+		WriteStencilActionPassParameters->RenderTargets.DepthStencil = FDepthStencilBinding(Inputs.StencilTexture, ERenderTargetLoadAction::ENoAction, ERenderTargetLoadAction::ELoad, FExclusiveDepthStencil::DepthNop_StencilWrite);
+
+		ClearStencilActionPassParameters = GraphBuilder.AllocParameters<FSlateRenderBatchParameters>(WriteStencilActionPassParameters);
+		ClearStencilActionPassParameters->RenderTargets.DepthStencil.SetStencilLoadAction(ERenderTargetLoadAction::EClear);
+	}
+
+	FSlateRenderBatchParameters* LastPassParameters  = NoneStencilActionPassParameters;
+	const FSlateClippingState*   LastClippingState   = nullptr;
+	const FSlateClippingOp*      LastClippingOp      = nullptr;
+
+	FSlateClippingCreateContext ClippingCreateContext;
+
+	TArray<FSlateRenderBatchOp, FRDGArrayAllocator> RenderBatchOps;
+	RenderBatchOps.Reserve(RenderBatches.Num());
+
+	const auto FlushDrawElementsPass = [&]
+	{
+		if (RenderBatchOps.IsEmpty())
+		{
+			return;
+		}
+
+		if (ERenderTargetLoadAction LoadAction = ConsumeLoadAction(ElementsLoadAction); LoadAction != ERenderTargetLoadAction::ELoad)
+		{
+			// Load action differs from the default read one, so make a copy and modify.
+			LastPassParameters = GraphBuilder.AllocParameters<FSlateRenderBatchParameters>(LastPassParameters);
+			LastPassParameters->RenderTargets[0].SetLoadAction(LoadAction);
+		}
+
+		const int32 NumRenderBatchOps = RenderBatchOps.Num();
+
+		FRDGPass* Pass = GraphBuilder.AddPass(RDG_EVENT_NAME("ElementBatch"), LastPassParameters, ERDGPassFlags::Raster, [Inputs = RenderBatchDrawInputs, RenderBatchOps = MoveTemp(RenderBatchOps)] (FRHICommandList& RHICmdList)
+		{
+			RHICmdList.SetViewport(Inputs->ElementsViewRect.Min.X, Inputs->ElementsViewRect.Min.Y, 0.0f, Inputs->ElementsViewRect.Max.X, Inputs->ElementsViewRect.Max.Y, 1.0f);
+			RHICmdList.SetScissorRect(false, 0, 0, 0, 0);
+
+			FSlateRenderBatchDrawState DrawState;
+			RHICmdList.ApplyCachedRenderTargets(DrawState.GraphicsPSOInit);
+			DrawState.GraphicsPSOInit.DepthStencilState = TStaticDepthStencilState<false, CF_Always>::GetRHI();
+
+			for (const FSlateRenderBatchOp& RenderBatchOp : RenderBatchOps)
+			{
+				DrawSlateRenderBatch(RHICmdList, DrawState, *Inputs, RenderBatchOp);
+			}
+		});
+
+		GraphBuilder.SetPassWorkload(Pass, NumRenderBatchOps);
+		check(RenderBatchOps.IsEmpty());
+	};
+
+	int32 NextRenderBatchIndex = FirstBatchIndex;
+
+	while (NextRenderBatchIndex != INDEX_NONE)
+	{
+		const FSlateRenderBatch& NextRenderBatch = RenderBatches[NextRenderBatchIndex];
+
+		NextRenderBatchIndex = NextRenderBatch.NextBatchIndex;
+
+		FSlateRenderBatchParameters* NextPassParameters  = LastPassParameters;
+		const FSlateClippingState*   NextClippingState   = NextRenderBatch.ClippingState;
+		const FSlateClippingOp*      NextClippingOp      = LastClippingOp;
+
+		if (NextClippingState != LastClippingState)
+		{
+			NextClippingOp = CreateSlateClipping(GraphBuilder, Inputs.ElementsOffset, NextClippingState, ClippingCreateContext);
+
+			switch (ClippingCreateContext.StencilAction)
+			{
+			case ESlateClippingStencilAction::Clear:
+				NextPassParameters = ClearStencilActionPassParameters;
+				break;
+
+			case ESlateClippingStencilAction::Write:
+				NextPassParameters = WriteStencilActionPassParameters;
+				break;
+
+			case ESlateClippingStencilAction::None:
+				NextPassParameters = NoneStencilActionPassParameters;
+				break;
+			}
+
+			LastClippingState = NextClippingState;
+			LastClippingOp = NextClippingOp;
+		}
+
+		const ESlateRenderBatchType NextRenderBatchType = GetSlateRenderBatchType(NextRenderBatch);
+
+		// Flush all primitive render batches when we encounter one that can't be added.
+		if (NextRenderBatchType != ESlateRenderBatchType::Primitive || NextPassParameters != LastPassParameters)
+		{
+			FlushDrawElementsPass();
+		}
+
+		LastPassParameters = NextPassParameters;
+
+		switch (NextRenderBatchType)
+		{
+		case ESlateRenderBatchType::CustomDrawer:
+		{
+			// Clear the color texture if we haven't done it yet.
+			if (ConsumeLoadAction(ElementsLoadAction) == ERenderTargetLoadAction::EClear)
+			{
+				AddClearRenderTargetPass(GraphBuilder, Inputs.ElementsTexture);
+			}
+
+			ICustomSlateElement::FDrawPassInputs DrawInputs;
+			DrawInputs.ElementsMatrix = Inputs.ElementsMatrix;
+			DrawInputs.ElementsOffset = Inputs.ElementsOffset;
+			DrawInputs.OutputTexture = ElementsTexture.Texture;
+			DrawInputs.SceneViewRect = Inputs.SceneViewRect;
+			DrawInputs.HDRDisplayColorGamut = Inputs.HDRDisplayColorGamut;
+			DrawInputs.UsedSlatePostBuffers = Inputs.UsedSlatePostBuffers;
+			DrawInputs.bOutputIsHDRDisplay = Inputs.bElementsTextureIsHDRDisplay;
+			DrawInputs.bWireFrame = Inputs.bWireframe;
+
+			NextRenderBatch.CustomDrawer->Draw_RenderThread(GraphBuilder, DrawInputs);
+
+			// Reset cached clipping state since custom draws mutate render state.
+			LastClippingState = nullptr;
+			break;
+		}
+		case ESlateRenderBatchType::PostProcess:
+		{
+			const FShaderParams& ShaderParams = NextRenderBatch.ShaderParams;
+
+			FSlatePostProcessBlurPassInputs BlurInputs;
+
+			if (Inputs.SceneViewportTexture != Inputs.ElementsTexture)
+			{
+				// Blur uses the scene viewport texture output as blur input and composites UI separately.
+				BlurInputs.InputTexture = Inputs.SceneViewportTexture;
+
+				if (HasBeenProduced(Inputs.ElementsTexture))
+				{
+					BlurInputs.SDRCompositeUITexture = Inputs.ElementsTexture;
+				}
+			}
+			else
+			{
+				// UI elements and scene are already composited together.
+				BlurInputs.InputTexture = Inputs.ElementsTexture;
+			}
+
+			BlurInputs.InputRect = FIntRect(ShaderParams.PixelParams.X, ShaderParams.PixelParams.Y, ShaderParams.PixelParams.Z, ShaderParams.PixelParams.W);
+			BlurInputs.OutputTexture = Inputs.SceneViewportTexture;
+			BlurInputs.OutputRect = BlurInputs.InputRect;
+			BlurInputs.OutputLoadAction = ConsumeLoadAction(ElementsLoadAction);
+			BlurInputs.ClippingOp = NextClippingOp;
+			BlurInputs.ClippingStencilBinding = &NextPassParameters->RenderTargets.DepthStencil;
+			BlurInputs.KernelSize = ShaderParams.PixelParams2.X;
+			BlurInputs.Strength = ShaderParams.PixelParams2.Y;
+			BlurInputs.DownsampleAmount = ShaderParams.PixelParams2.Z;
+			BlurInputs.CornerRadius = ShaderParams.PixelParams3;
+
+			AddSlatePostProcessBlurPass(GraphBuilder, BlurInputs);
+			break;
+		}
+		case ESlateRenderBatchType::Primitive:
+		{
+			const FSlateSceneView* SceneView = nullptr;
+
+			if (NextRenderBatch.ShaderResource && NextRenderBatch.ShaderResource->GetType() == ESlateShaderResource::Material)
+			{
+				// Skip material render batches when the engine is not available.
+				if (!GEngine)
+				{
+					break;
+				}
+
+				const FSlateSceneViewAllocateInputs AllocateInputs =
+				{
+					  .TextureExtent        = ElementsTextureExtent
+					, .ViewRect             = Inputs.SceneViewRect
+					, .ViewProjectionMatrix = Inputs.ElementsMatrix
+					, .CursorPosition       = Inputs.CursorPosition
+					, .Time                 = Inputs.Time
+					, .ViewportScaleUI      = Inputs.ViewportScaleUI
+				};
+
+				SceneView = SceneViewAllocator->BeginAllocateSceneView(GraphBuilder, ResourceManager, NextRenderBatch.SceneIndex, AllocateInputs);
+			}
+
+			RenderBatchOps.Emplace(&NextRenderBatch, NextClippingOp, SceneView);
+			break;
+		}
+		default: checkNoEntry();
+		}
+	}
+
+	if (!RenderBatchOps.IsEmpty())
+	{
+		FlushDrawElementsPass();
+	}
+
+	// If no batches were rendered at all, then we might need to just clear the render target.
+	if (ConsumeLoadAction(ElementsLoadAction) == ERenderTargetLoadAction::EClear)
+	{
+		AddClearRenderTargetPass(GraphBuilder, Inputs.ElementsTexture);
+	}
+	else
+	{
+		// Don't do color correction on iOS or Android, we don't have the GPU overhead for it.
+#if !(PLATFORM_IOS || PLATFORM_ANDROID)
+		if (Inputs.bAllowColorDeficiencyCorrection && GSlateColorDeficiencyType != EColorVisionDeficiency::NormalVision && GSlateColorDeficiencySeverity > 0)
+		{
+			FSlatePostProcessColorDeficiencyPassInputs ColorDeficiencyInputs;
+			ColorDeficiencyInputs.InputTexture = ElementsTexture;
+			ColorDeficiencyInputs.OutputTexture = ElementsTexture;
+
+			AddSlatePostProcessColorDeficiencyPass(GraphBuilder, ColorDeficiencyInputs);
+		}
+#endif
+	}
 	
+	INC_DWORD_STAT_BY(STAT_SlateScissorClips, ClippingCreateContext.NumScissors);
+	INC_DWORD_STAT_BY(STAT_SlateStencilClips, ClippingCreateContext.NumStencils);
+}
\ No newline at end of file
