Clean-up of Slate RHI renderer and conversion to RDG.

diff --git a/Engine/Source/Runtime/SlateRHIRenderer/Private/SlateRHIRenderingPolicy.h b/Engine/Source/Runtime/SlateRHIRenderer/Private/SlateRHIRenderingPolicy.h
--- a/Engine/Source/Runtime/SlateRHIRenderer/Private/SlateRHIRenderingPolicy.h
+++ b/Engine/Source/Runtime/SlateRHIRenderer/Private/SlateRHIRenderingPolicy.h
@@ -29,102 +29,55 @@ class FSlateMaterialShaderVS;
 struct FMaterialShaderTypes;
 struct IPooledRenderTarget;
 
-struct FSlateRenderingParams
-{
-	FMatrix44f ViewProjectionMatrix;
-	FVector2f ViewOffset;
-	FIntRect ViewRect;
-	FIntPoint CursorPostion;
-	FGameTime Time;
-	TRefCountPtr<IPooledRenderTarget> UITarget;
-	EDisplayColorGamut HDRDisplayColorGamut;
-	float ViewportScaleUI;
-	ESlatePostRT UsedSlatePostBuffers;
-	bool bWireFrame;
-	bool bIsHDR;
-
-	FSlateRenderingParams(const FMatrix& InViewProjectionMatrix, FGameTime InTime)
-		: ViewProjectionMatrix(InViewProjectionMatrix)
-		, ViewOffset(0.f, 0.f)
-		, ViewRect(FIntRect())
-		, CursorPostion(FIntPoint())
-		, Time(InTime)
-		, HDRDisplayColorGamut(EDisplayColorGamut::sRGB_D65)
-		, ViewportScaleUI(1.0f)
-		, UsedSlatePostBuffers(ESlatePostRT::None)
-		, bWireFrame(false)
-		, bIsHDR(false)
-	{
-	}
-};
-
-class FSlateRHIRenderingPolicy : public FSlateRenderingPolicy
+class FSlateRHIRenderingPolicy final : public FSlateRenderingPolicy
 {
 public:
-	FSlateRHIRenderingPolicy(TSharedRef<FSlateFontServices> InSlateFontServices, TSharedRef<FSlateRHIResourceManager> InResourceManager, TOptional<int32> InitialBufferSize = TOptional<int32>());
-
-	void BuildRenderingBuffers(FRHICommandListImmediate& RHICmdList, FSlateBatchData& InBatchData);
-
-	void DrawElements(FRHICommandListImmediate& RHICmdList, class FSlateBackBuffer& BackBuffer, FTextureRHIRef& ColorTarget, FTextureRHIRef& PostProcessBuffer,
-		FTextureRHIRef& DepthStencilTarget, int32 FirstBatchIndex,  const TArray<FSlateRenderBatch>& RenderBatches, const FSlateRenderingParams& Params);
+	FSlateRHIRenderingPolicy(TSharedRef<FSlateFontServices> InSlateFontServices, TSharedRef<FSlateRHIResourceManager> InResourceManager);
 
 	virtual TSharedRef<FSlateShaderResourceManager> GetResourceManager() const override { return ResourceManager; }
 	virtual bool IsVertexColorInLinearSpace() const override { return false; }
-	
-	void InitResources();
-	void ReleaseResources();
 
-	void BeginDrawingWindows();
-	void EndDrawingWindows();
-
-	void SetUseGammaCorrection( bool bInUseGammaCorrection ) { bGammaCorrect = bInUseGammaCorrection; }
-	void SetApplyColorDeficiencyCorrection(bool bInApplyColorCorrection) { bApplyColorDeficiencyCorrection = bInApplyColorCorrection; }
-
-	void TickPostProcessResources();
-
-	bool GetApplyColorDeficiencyCorrection() const { return bApplyColorDeficiencyCorrection; }
+	FSlateRHIResourceManager& GetResourceManagerRHI() const { return ResourceManager.Get(); }
 
 	virtual void AddSceneAt(FSceneInterface* Scene, int32 Index) override;
 	virtual void ClearScenes() override;
 
-	virtual void FlushGeneratedResources();
-
-	void BlurRectExternal(FRHICommandListImmediate& RHICmdList, FRHITexture* BlurSrc, FRHITexture* BlurDst, FIntRect SrcRect, FIntRect DstRect, float BlurStrength) const;
-
-private:
-	ETextureSamplerFilter GetSamplerFilter(const UTexture* Texture) const;
-
-	/**
-	 * Returns the pixel shader that should be used for the specified ShaderType and DrawEffects
-	 * 
-	 * @param ShaderType	The shader type being used
-	 * @param DrawEffects	Draw effects being used
-	 * @return The pixel shader for use with the shader type and draw effects
-	 */
-	TShaderRef<FSlateElementPS> GetTexturePixelShader(FGlobalShaderMap* ShaderMap, ESlateShader ShaderType, ESlateDrawEffect DrawEffects, bool bUseTextureGrayscale,  bool bIsVirtualTexture);
-	void ChooseMaterialShaderTypes(ESlateShader ShaderType, bool bUseInstancing, FMaterialShaderTypes& OutShaderTypes);
-
-	/** @return The RHI primitive type from the Slate primitive type */
-	EPrimitiveType GetRHIPrimitiveType(ESlateDrawPrimitive SlateType);
-
 private:
-	/** Buffers used for rendering */
-	TSlateElementVertexBuffer<FSlateVertex> SourceVertexBuffer;
-	FSlateElementIndexBuffer SourceIndexBuffer;
-
-	FSlateStencilClipVertexBuffer StencilVertexBuffer;
-
-	/** Handles post process effects for slate */
-	TSharedRef<FSlatePostProcessor> PostProcessor;
-
 	TSharedRef<FSlateRHIResourceManager> ResourceManager;
+};
 
-	bool bGammaCorrect;
-	bool bApplyColorDeficiencyCorrection;
-
-	TOptional<int32> InitialBufferSizeOverride;
+struct FSlateElementsBuffers
+{
+	FRDGBuffer* VertexBuffer = nullptr;
+	FRDGBuffer* IndexBuffer = nullptr;
+};
 
-	TArray<FTextureLODGroup> TextureLODGroups;
+FSlateElementsBuffers BuildSlateElementsBuffers(FRDGBuilder& GraphBuilder, FSlateBatchData& BatchData);
 
-	UDeviceProfile* LastDeviceProfile;
+struct FSlateDrawElementsPassInputs
+{
+	FRDGTexture* StencilTexture = nullptr;
+	FRDGTexture* ElementsTexture = nullptr;
+	FRDGTexture* SceneViewportTexture = nullptr;
+	ERenderTargetLoadAction ElementsLoadAction = ERenderTargetLoadAction::ENoAction;
+	FSlateElementsBuffers ElementsBuffers;
+	FMatrix44f ElementsMatrix;
+	FVector2f ElementsOffset = FVector2f::Zero();
+	FIntRect SceneViewRect;
+	FIntPoint CursorPosition = FIntPoint::ZeroValue;
+	FGameTime Time;
+	EDisplayColorGamut HDRDisplayColorGamut = EDisplayColorGamut::sRGB_D65;
+	ESlatePostRT UsedSlatePostBuffers = ESlatePostRT::None;
+	float ViewportScaleUI = 1.0f;
+	bool bWireframe = false;
+	bool bElementsTextureIsHDRDisplay = false;
+	bool bAllowGammaCorrection = true;
+	bool bAllowColorDeficiencyCorrection = true;
 };
+
+void AddSlateDrawElementsPass(
+	FRDGBuilder& GraphBuilder,
+	const FSlateRHIRenderingPolicy& RenderingPolicy,
+	const FSlateDrawElementsPassInputs& Inputs,
+	TConstArrayView<FSlateRenderBatch> RenderBatches,
+	int32 FirstBatchIndex);
\ No newline at end of file
diff --git a/Engine/Source/Runtime/SlateRHIRenderer/Private/SlateRHIRenderingPolicyInterface.cpp b/Engine/Source/Runtime/SlateRHIRenderer/Private/SlateRHIRenderingPolicyInterface.cpp
--- a/Engine/Source/Runtime/SlateRHIRenderer/Private/SlateRHIRenderingPolicyInterface.cpp
+++ b/Engine/Source/Runtime/SlateRHIRenderer/Private/SlateRHIRenderingPolicyInterface.cpp
@@ -1,48 +1,3 @@
 // Copyright Epic Games, Inc. All Rights Reserved.
 
-#include "Interfaces/SlateRHIRenderingPolicyInterface.h"
-#include "SlateRHIRenderingPolicy.h"
-#include "SlateRHIRenderer.h"
-
-FSlateRHIRenderingPolicyInterface::FSlateRHIRenderingPolicyInterface(FSlateRHIRenderingPolicy* InRenderingPolicy)
-	: RenderingPolicy(InRenderingPolicy)
-{
-}
-
-int32 FSlateRHIRenderingPolicyInterface::GetProcessSlatePostBuffers()
-{
-	return FSlateRHIRenderer::GetProcessSlatePostBuffers();
-}
-
-bool FSlateRHIRenderingPolicyInterface::IsValid() const
-{
-	return RenderingPolicy != nullptr;
-}
-
-bool FSlateRHIRenderingPolicyInterface::IsVertexColorInLinearSpace() const
-{
-	if (RenderingPolicy)
-	{
-		return RenderingPolicy->IsVertexColorInLinearSpace();
-	}
-
-	return false;
-}
-
-bool FSlateRHIRenderingPolicyInterface::GetApplyColorDeficiencyCorrection() const
-{
-	if (RenderingPolicy)
-	{
-		return RenderingPolicy->GetApplyColorDeficiencyCorrection();
-	}
-
-	return false;
-}
-
-void FSlateRHIRenderingPolicyInterface::BlurRectExternal(FRHICommandListImmediate& RHICmdList, FRHITexture* BlurSrc, FRHITexture* BlurDst, FIntRect SrcRect, FIntRect DstRect, float BlurStrength) const
-{
-	if (RenderingPolicy)
-	{
-		RenderingPolicy->BlurRectExternal(RHICmdList, BlurSrc, BlurDst, SrcRect, DstRect, BlurStrength);
-	}
-}
\ No newline at end of file
+#include "Interfaces/SlateRHIRenderingPolicyInterface.h"
\ No newline at end of file
diff --git a/Engine/Source/Runtime/SlateRHIRenderer/Private/SlateShaders.cpp b/Engine/Source/Runtime/SlateRHIRenderer/Private/SlateShaders.cpp
--- a/Engine/Source/Runtime/SlateRHIRenderer/Private/SlateShaders.cpp
+++ b/Engine/Source/Runtime/SlateRHIRenderer/Private/SlateShaders.cpp
@@ -17,15 +17,8 @@ IMPLEMENT_SHADER_TYPE(, FSlateElementVS, TEXT("/Engine/Private/SlateVertexShader
 
 IMPLEMENT_SHADER_TYPE(, FSlateDebugOverdrawPS, TEXT("/Engine/Private/SlateElementPixelShader.usf"), TEXT("DebugOverdrawMain"), SF_Pixel );
 
-IMPLEMENT_SHADER_TYPE(, FSlatePostProcessBlurPS, TEXT("/Engine/Private/SlatePostProcessPixelShader.usf"), TEXT("GaussianBlurMain"), SF_Pixel);
-IMPLEMENT_SHADER_TYPE(, FSlatePostProcessDownsamplePS, TEXT("/Engine/Private/SlatePostProcessPixelShader.usf"), TEXT("DownsampleMain"), SF_Pixel);
-IMPLEMENT_SHADER_TYPE(template<>, FSlatePostProcessUpsamplePS<ESlatePostProcessUpsamplePSPermutation::SDR>, TEXT("/Engine/Private/SlatePostProcessPixelShader.usf"), TEXT("UpsampleMain"), SF_Pixel);
-IMPLEMENT_SHADER_TYPE(template<>, FSlatePostProcessUpsamplePS<ESlatePostProcessUpsamplePSPermutation::HDR_SCRGB>, TEXT("/Engine/Private/SlatePostProcessPixelShader.usf"), TEXT("UpsampleMain"), SF_Pixel);
-IMPLEMENT_SHADER_TYPE(template<>, FSlatePostProcessUpsamplePS<ESlatePostProcessUpsamplePSPermutation::HDR_PQ10>, TEXT("/Engine/Private/SlatePostProcessPixelShader.usf"), TEXT("UpsampleMain"), SF_Pixel);
-IMPLEMENT_SHADER_TYPE(, FSlatePostProcessColorDeficiencyPS, TEXT("/Engine/Private/SlatePostProcessColorDeficiencyPixelShader.usf"), TEXT("ColorDeficiencyMain"), SF_Pixel);
-
-IMPLEMENT_SHADER_TYPE(, FSlateMaskingVS, TEXT("/Engine/Private/SlateMaskingShader.usf"), TEXT("MainVS"), SF_Vertex);
-IMPLEMENT_SHADER_TYPE(, FSlateMaskingPS, TEXT("/Engine/Private/SlateMaskingShader.usf"), TEXT("MainPS"), SF_Pixel);
+IMPLEMENT_GLOBAL_SHADER(FSlateMaskingVS, "/Engine/Private/SlateMaskingShader.usf", "MainVS", SF_Vertex);
+IMPLEMENT_GLOBAL_SHADER(FSlateMaskingPS, "/Engine/Private/SlateMaskingShader.usf", "MainPS", SF_Pixel);
 
 IMPLEMENT_SHADER_TYPE(, FSlateDebugBatchingPS, TEXT("/Engine/Private/SlateElementPixelShader.usf"), TEXT("DebugBatchingMain"), SF_Pixel );
 
@@ -33,14 +26,6 @@ IMPLEMENT_SHADER_TYPE(, FSlateDebugBatchingPS, TEXT("/Engine/Private/SlateElemen
 	typedef TSlateElementPS<ESlateShader::ShaderType,bDrawDisabledEffect,bUseTextureAlpha, bUseTextureGrayscale, bIsVirtualTexture> TSlateElementPS##ShaderType##bDrawDisabledEffect##bUseTextureAlpha##bUseTextureGrayscale##bIsVirtualTexture##A; \
 	IMPLEMENT_SHADER_TYPE(template<>,TSlateElementPS##ShaderType##bDrawDisabledEffect##bUseTextureAlpha##bUseTextureGrayscale##bIsVirtualTexture##A,TEXT("/Engine/Private/SlateElementPixelShader.usf"),TEXT("Main"),SF_Pixel);
 
-#if WITH_EDITOR
-bool FHDREditorConvertPS::ShouldCompilePermutation(const FGlobalShaderPermutationParameters& Parameters)
-{
-	return IsFeatureLevelSupported(Parameters.Platform, ERHIFeatureLevel::SM5);
-}
-
-IMPLEMENT_SHADER_TYPE(, FHDREditorConvertPS, TEXT("/Engine/Private/CompositeUIPixelShader.usf"), TEXT("HDREditorConvert"), SF_Pixel);
-#endif
 /**
 * All the different permutations of shaders used by slate. Uses #defines to avoid dynamic branches
 */
@@ -175,50 +160,4 @@ void FSlateElementVS::SetViewProjection(FRHIBatchedShaderParameters& BatchedPara
 void FSlateElementVS::SetShaderParameters(FRHIBatchedShaderParameters& BatchedParameters, const FVector4f& ShaderParams )
 {
 	SetShaderValue(BatchedParameters, VertexShaderParams, ShaderParams );
-}
-
-/** Serializes the shader data */
-/*bool FSlateElementVS::Serialize( FArchive& Ar )
-{
-	bool bShaderHasOutdatedParameters = FGlobalShader::Serialize( Ar );
-
-	Ar << ViewProjection;
-	Ar << VertexShaderParams;
-
-	return bShaderHasOutdatedParameters;
-}*/
-
-
-/************************************************************************/
-/* FSlateMaskingVS                                            */
-/************************************************************************/
-
-FSlateMaskingVS::FSlateMaskingVS(const ShaderMetaType::CompiledShaderInitializerType& Initializer)
-	: FGlobalShader(Initializer)
-{
-	ViewProjection.Bind(Initializer.ParameterMap, TEXT("ViewProjection"));
-	MaskRect.Bind(Initializer.ParameterMap, TEXT("MaskRectPacked"));
-}
-
-void FSlateMaskingVS::SetViewProjection(FRHIBatchedShaderParameters& BatchedParameters, const FMatrix44f& InViewProjection)
-{
-	SetShaderValue(BatchedParameters, ViewProjection, InViewProjection);
-}
-
-void FSlateMaskingVS::SetMaskRect(FRHIBatchedShaderParameters& BatchedParameters, const FVector2f TopLeft, const FVector2f TopRight, const FVector2f BotLeft, const FVector2f BotRight)
-{
-	FVector4f MaskRectVal[2] = { FVector4f(TopLeft, TopRight), FVector4f(BotLeft, BotRight) };
-
-	SetShaderValue(BatchedParameters, MaskRect, MaskRectVal);
-}
-
-/** Serializes the shader data */
-/*bool FSlateMaskingVS::Serialize(FArchive& Ar)
-{
-	bool bShaderHasOutdatedParameters = FGlobalShader::Serialize(Ar);
-
-	Ar << ViewProjection;
-	Ar << MaskRect;
-
-	return bShaderHasOutdatedParameters;
-}*/
+}
\ No newline at end of file
diff --git a/Engine/Source/Runtime/SlateRHIRenderer/Private/SlateShaders.h b/Engine/Source/Runtime/SlateRHIRenderer/Private/SlateShaders.h
--- a/Engine/Source/Runtime/SlateRHIRenderer/Private/SlateShaders.h
+++ b/Engine/Source/Runtime/SlateRHIRenderer/Private/SlateShaders.h
@@ -13,6 +13,7 @@
 #include "RHIStaticStates.h"
 #include "TextureResource.h"
 #include "RenderUtils.h"
+#include "ShaderParameterStruct.h"
 
 extern EColorVisionDeficiency GSlateColorDeficiencyType;
 extern int32 GSlateColorDeficiencySeverity;
@@ -314,276 +315,22 @@ private:
 	LAYOUT_FIELD(FShaderParameter, BatchColor);
 };
 
-const int32 MAX_BLUR_SAMPLES = 127;
-
-class FSlatePostProcessBlurPS : public FSlateElementPS
-{
-	DECLARE_SHADER_TYPE(FSlatePostProcessBlurPS, Global);
-public:
-	/** Indicates that this shader should be cached */
-	static bool ShouldCompilePermutation(const FGlobalShaderPermutationParameters& Parameters)
-	{
-		return true;
-	}
-
-	FSlatePostProcessBlurPS()
-	{
-	}
-
-	/** Constructor.  Binds all parameters used by the shader */
-	FSlatePostProcessBlurPS(const ShaderMetaType::CompiledShaderInitializerType& Initializer)
-		: FSlateElementPS(Initializer)
-	{
-		BufferSizeAndDirection.Bind(Initializer.ParameterMap, TEXT("BufferSizeAndDirection"));
-		WeightAndOffsets.Bind(Initializer.ParameterMap, TEXT("WeightAndOffsets"));
-		SampleCount.Bind(Initializer.ParameterMap, TEXT("SampleCount"));
-		UVBounds.Bind(Initializer.ParameterMap, TEXT("UVBounds"));
-	}
-
-	void SetBufferSizeAndDirection(FRHIBatchedShaderParameters& BatchedParameters, const FVector2f InBufferSize, const FVector2f InDir)
-	{
-		SetShaderValue(BatchedParameters, BufferSizeAndDirection, FVector4f(InBufferSize, InDir));
-	}
-
-	void SetWeightsAndOffsets(FRHIBatchedShaderParameters& BatchedParameters, const TArray<FVector4f>& InWeightsAndOffsets, int32 NumSamples )
-	{
-		check(InWeightsAndOffsets.Num() <= MAX_BLUR_SAMPLES);
-		SetShaderValueArray<FVector4f>(BatchedParameters, WeightAndOffsets, InWeightsAndOffsets.GetData(), InWeightsAndOffsets.Num() );
-		SetShaderValue(BatchedParameters, SampleCount, NumSamples);
-	}
-
-	void SetUVBounds(FRHIBatchedShaderParameters& BatchedParameters, const FVector4f& InUVBounds)
-	{
-		SetShaderValue(BatchedParameters, UVBounds, InUVBounds);
-	}
-
-private:
-	LAYOUT_FIELD(FShaderParameter, BufferSizeAndDirection);
-	LAYOUT_FIELD(FShaderParameter, WeightAndOffsets);
-	LAYOUT_FIELD(FShaderParameter, SampleCount);
-	LAYOUT_FIELD(FShaderParameter, UVBounds);
-};
-
-
-class FSlatePostProcessDownsamplePS : public FSlateElementPS
-{
-	DECLARE_SHADER_TYPE(FSlatePostProcessDownsamplePS, Global);
-public:
-	/** Indicates that this shader should be cached */
-	static bool ShouldCompilePermutation(const FGlobalShaderPermutationParameters& Parameters)
-	{
-		return true;
-	}
-
-	FSlatePostProcessDownsamplePS()
-	{
-	}
-
-	/** Constructor.  Binds all parameters used by the shader */
-	FSlatePostProcessDownsamplePS(const ShaderMetaType::CompiledShaderInitializerType& Initializer)
-		: FSlateElementPS(Initializer)
-	{
-		UVBounds.Bind(Initializer.ParameterMap, TEXT("UVBounds"));
-	}
-
-	void SetUVBounds(FRHIBatchedShaderParameters& BatchedParameters, const FVector4f& InUVBounds)
-	{
-		SetShaderValue(BatchedParameters, UVBounds, InUVBounds);
-	}
-
-private:
-	LAYOUT_FIELD(FShaderParameter, UVBounds);
-};
-
-enum class ESlatePostProcessUpsamplePSPermutation
-{
-	SDR = 0,
-	HDR_SCRGB,
-	HDR_PQ10,
-};
-
-template<ESlatePostProcessUpsamplePSPermutation SlatePostProcessUpsamplePSPermutation>
-class FSlatePostProcessUpsamplePS : public FSlateElementPS
-{
-	DECLARE_SHADER_TYPE(FSlatePostProcessUpsamplePS, Global);
-public:
-	/** Indicates that this shader should be cached */
-	static bool ShouldCompilePermutation(const FGlobalShaderPermutationParameters& Parameters)
-	{
-		return true;
-	}
-
-	FSlatePostProcessUpsamplePS()
-	{
-	}
-
-	/** Constructor.  Binds all parameters used by the shader */
-	FSlatePostProcessUpsamplePS(const ShaderMetaType::CompiledShaderInitializerType& Initializer)
-		: FSlateElementPS(Initializer)
-	{
-	}
-
-	/**
-	 * Modifies the compilation of this shader
-	 */
-	static void ModifyCompilationEnvironment(const FGlobalShaderPermutationParameters& Parameters, FShaderCompilerEnvironment& OutEnvironment)
-	{
-		// Set defines based on what this shader will be used for
-		OutEnvironment.SetDefine(TEXT("OUTPUT_TO_UI_TARGET"), (uint32)(SlatePostProcessUpsamplePSPermutation != ESlatePostProcessUpsamplePSPermutation::SDR ? 1 : 0));
-		OutEnvironment.SetDefine(TEXT("SCRGB_ENCODING"), SlatePostProcessUpsamplePSPermutation == ESlatePostProcessUpsamplePSPermutation::HDR_SCRGB);
-		FSlateElementPS::ModifyCompilationEnvironment(Parameters, OutEnvironment);
-	}
-
-private:
-};
-
-
-class FSlatePostProcessColorDeficiencyPS : public FSlateElementPS
-{
-	DECLARE_SHADER_TYPE(FSlatePostProcessColorDeficiencyPS, Global);
-public:
-	/** Indicates that this shader should be cached */
-	static bool ShouldCompilePermutation(const FGlobalShaderPermutationParameters& Parameters)
-	{
-		return true;
-	}
-
-	FSlatePostProcessColorDeficiencyPS()
-	{
-	}
-
-	/** Constructor.  Binds all parameters used by the shader */
-	FSlatePostProcessColorDeficiencyPS(const ShaderMetaType::CompiledShaderInitializerType& Initializer)
-		: FSlateElementPS(Initializer)
-	{
-		ColorVisionDeficiencyType.Bind(Initializer.ParameterMap, TEXT("ColorVisionDeficiencyType"));
-		ColorVisionDeficiencySeverity.Bind(Initializer.ParameterMap, TEXT("ColorVisionDeficiencySeverity"));
-		bCorrectDeficiency.Bind(Initializer.ParameterMap, TEXT("bCorrectDeficiency"));
-		bSimulateCorrectionWithDeficiency.Bind(Initializer.ParameterMap, TEXT("bSimulateCorrectionWithDeficiency"));
-	}
-
-	void SetColorRules(FRHIBatchedShaderParameters& BatchedParameters, bool bCorrect, EColorVisionDeficiency DeficiencyType, int32 Severity)
-	{
-		SetShaderValue(BatchedParameters, ColorVisionDeficiencyType, (float)DeficiencyType);
-		SetShaderValue(BatchedParameters, ColorVisionDeficiencySeverity, (float)Severity);
-		SetShaderValue(BatchedParameters, bCorrectDeficiency, bCorrect ? 1.0f : 0.0f);
-	}
-
-	void SetShowCorrectionWithDeficiency(FRHIBatchedShaderParameters& BatchedParameters, bool bShowCorrectionWithDeficiency)
-	{
-		SetShaderValue(BatchedParameters, bSimulateCorrectionWithDeficiency, bShowCorrectionWithDeficiency ? 1.0f : 0.0f);
-	}
-
-private:
-	LAYOUT_FIELD(FShaderParameter, ColorVisionDeficiencyType);
-	LAYOUT_FIELD(FShaderParameter, ColorVisionDeficiencySeverity);
-	LAYOUT_FIELD(FShaderParameter, bCorrectDeficiency);
-	LAYOUT_FIELD(FShaderParameter, bSimulateCorrectionWithDeficiency);
-};
-
-
 class FSlateMaskingVS : public FGlobalShader
 {
-	DECLARE_SHADER_TYPE(FSlateMaskingVS, Global);
 public:
-	/** Indicates that this shader should be cached */
-	static bool ShouldCompilePermutation(const FGlobalShaderPermutationParameters& Parameters)
-	{
-		return true;
-	}
-
-	FSlateMaskingVS()
-	{
-	}
-
-	/** Constructor.  Binds all parameters used by the shader */
-	FSlateMaskingVS(const ShaderMetaType::CompiledShaderInitializerType& Initializer);
+	DECLARE_GLOBAL_SHADER(FSlateMaskingVS);
+	SHADER_USE_PARAMETER_STRUCT(FSlateMaskingVS, FGlobalShader);
 
-	/**
-	* Sets the view projection parameter
-	*
-	* @param InViewProjection	The ViewProjection matrix to use when this shader is bound
-	*/
-	void SetViewProjection(FRHIBatchedShaderParameters& BatchedParameters, const FMatrix44f& InViewProjection);
-
-	/**
-	 * Sets the mask rect positions
-	 */
-	void SetMaskRect(FRHIBatchedShaderParameters& BatchedParameters, const FVector2f TopLeft, const FVector2f TopRight, const FVector2f BotLeft, const FVector2f BotRight);
-
-	//virtual bool Serialize(FArchive& Ar) override;
-
-private:
-	/** Mask rect parameter */
-	LAYOUT_FIELD(FShaderParameter, MaskRect)
-	/** ViewProjection parameter used by the shader */
-	LAYOUT_FIELD(FShaderParameter, ViewProjection)
+	BEGIN_SHADER_PARAMETER_STRUCT(FParameters, )
+		SHADER_PARAMETER_ARRAY(FVector4f, MaskRectPacked, [2])
+	END_SHADER_PARAMETER_STRUCT()
 };
 
 class FSlateMaskingPS : public FGlobalShader
 {
-	DECLARE_SHADER_TYPE(FSlateMaskingPS, Global);
-public:
-	/** Indicates that this shader should be cached */
-	static bool ShouldCompilePermutation(const FGlobalShaderPermutationParameters& Parameters)
-	{
-		return true;
-	}
-
-	FSlateMaskingPS()
-	{
-	}
-
-	/** Constructor.  Binds all parameters used by the shader */
-	FSlateMaskingPS(const ShaderMetaType::CompiledShaderInitializerType& Initializer)
-		: FGlobalShader(Initializer)
-	{
-	}
-};
-
-
-#if WITH_EDITOR
-// Pixel shader to convert UI from linear rec709 to PQ 2020 for HDR monitors
-class FHDREditorConvertPS : public FGlobalShader
-{
-	DECLARE_SHADER_TYPE(FHDREditorConvertPS, Global);
 public:
-
-	FHDREditorConvertPS(const ShaderMetaType::CompiledShaderInitializerType& Initializer) :
-		FGlobalShader(Initializer)
-	{
-		SceneTexture.Bind(Initializer.ParameterMap, TEXT("SceneTexture"));
-		SceneSampler.Bind(Initializer.ParameterMap, TEXT("SceneSampler"));
-		UILevel.Bind(Initializer.ParameterMap, TEXT("UILevel"));
-	}
-	FHDREditorConvertPS() {}
-
-	void SetParameters(FRHIBatchedShaderParameters& BatchedParameters, FRHITexture* SceneTextureRHI)
-	{
-		SetTextureParameter(BatchedParameters, SceneTexture, SceneSampler, TStaticSamplerState<SF_Point>::GetRHI(), SceneTextureRHI);
-		
-		static auto CVarHDRNITLevel = IConsoleManager::Get().FindConsoleVariable(TEXT("Editor.HDRNITLevel"));
-		SetShaderValue(BatchedParameters, UILevel, CVarHDRNITLevel->GetFloat());
-	}
-
-	static bool ShouldCompilePermutation(const FGlobalShaderPermutationParameters& Parameters);
-
-	static const TCHAR* GetSourceFilename()
-	{
-		return TEXT("/Engine/Private/CompositeUIPixelShader.usf");
-	}
-
-	static const TCHAR* GetFunctionName()
-	{
-		return TEXT("HDREditorConvert");
-	}
-
-private:
-	LAYOUT_FIELD(FShaderResourceParameter, SceneTexture);
-	LAYOUT_FIELD(FShaderResourceParameter, SceneSampler);
-	LAYOUT_FIELD(FShaderParameter, UILevel);
+	DECLARE_GLOBAL_SHADER(FSlateMaskingPS);
 };
-#endif
 
 /** The simple element vertex declaration. */
 extern TGlobalResource<FSlateVertexDeclaration> GSlateVertexDeclaration;
@@ -592,4 +339,4 @@ extern TGlobalResource<FSlateVertexDeclaration> GSlateVertexDeclaration;
 extern TGlobalResource<FSlateInstancedVertexDeclaration> GSlateInstancedVertexDeclaration;
 
 /** The vertex declaration for rendering stencil masks. */
-extern TGlobalResource<FSlateMaskingVertexDeclaration> GSlateMaskingVertexDeclaration;
+extern TGlobalResource<FSlateMaskingVertexDeclaration> GSlateMaskingVertexDeclaration;
\ No newline at end of file
diff --git a/Engine/Source/Runtime/SlateRHIRenderer/Public/FX/SlatePostBufferBlur.h b/Engine/Source/Runtime/SlateRHIRenderer/Public/FX/SlatePostBufferBlur.h
--- a/Engine/Source/Runtime/SlateRHIRenderer/Public/FX/SlatePostBufferBlur.h
+++ b/Engine/Source/Runtime/SlateRHIRenderer/Public/FX/SlatePostBufferBlur.h
@@ -2,12 +2,9 @@
 
 #pragma once
 
-#include "CoreMinimal.h"
 #include "FX/SlateRHIPostBufferProcessor.h"
-
 #include "SlatePostBufferBlur.generated.h"
 
-
 /**
  * Proxy for post buffer processor that the renderthread uses to perform processing
  * This proxy exists because generally speaking usage on UObjects on the renderthread
@@ -19,7 +16,7 @@ class SLATERHIRENDERER_API FSlatePostBufferBlurProxy : public FSlateRHIPostBuffe
 public:
 
 	//~ Begin FSlateRHIPostBufferProcessorProxy Interface
-	virtual void PostProcess_Renderthread(FRHICommandListImmediate& RHICmdList, FRHITexture* Src, FRHITexture* Dst, FIntRect SrcRect, FIntRect DstRect, FSlateRHIRenderingPolicyInterface InRenderingPolicy) override;
+	virtual void PostProcess_Renderthread(FRDGBuilder& GraphBuilder, const FScreenPassTexture& InputTexture, const FScreenPassTexture& OutputTexture) override;
 	virtual void OnUpdateValuesRenderThread() override;
 	//~ End FSlateRHIPostBufferProcessorProxy Interface
 
@@ -59,10 +56,7 @@ public:
 	USlatePostBufferBlur();
 	virtual ~USlatePostBufferBlur() override;
 
-	//~ Begin USlateRHIPostBufferProcessor Interface
-	virtual void PostProcess(FRenderResource* InViewInfo, FRenderResource* InViewportTexture, FVector2D InElementWindowSize, FSlateRHIRenderingPolicyInterface InRenderingPolicy, UTextureRenderTarget2D* InSlatePostBuffer) override;
 	virtual TSharedPtr<FSlateRHIPostBufferProcessorProxy> GetRenderThreadProxy();
-	//~ End USlateRHIPostBufferProcessor Interface
 
 private:
 
diff --git a/Engine/Source/Runtime/SlateRHIRenderer/Public/FX/SlateRHIPostBufferProcessor.h b/Engine/Source/Runtime/SlateRHIRenderer/Public/FX/SlateRHIPostBufferProcessor.h
--- a/Engine/Source/Runtime/SlateRHIRenderer/Public/FX/SlateRHIPostBufferProcessor.h
+++ b/Engine/Source/Runtime/SlateRHIRenderer/Public/FX/SlateRHIPostBufferProcessor.h
@@ -2,16 +2,13 @@
 
 #pragma once
 
-#include "CoreMinimal.h"
 #include "Engine/TextureRenderTarget2D.h"
-#include "RenderResource.h"
-#include "RendererInterface.h"
 #include "Rendering/RenderingCommon.h"
 #include "Rendering/SlateRenderer.h"
-#include "RHIFwd.h"
 #include "UObject/Object.h"
 #include "UObject/ObjectMacros.h"
 #include "Interfaces/SlateRHIRenderingPolicyInterface.h"
+#include "ScreenPass.h"
 
 #include "SlateRHIPostBufferProcessor.generated.h"
 
@@ -26,23 +23,21 @@ class SLATERHIRENDERER_API FSlateRHIPostBufferProcessorProxy : public TSharedFro
 {
 
 public:
-
 	virtual ~FSlateRHIPostBufferProcessorProxy()
 	{
 	}
 
-	/**
-	 * Called directly inside renderthread to perform some processing, do not enque commands here as we are already in the renderthread
-	 *
-	 * @param Src					Source texture to process as Input
-	 * @param Dst					Destination texture to store process output
-	 * @param SrcRect				Rect within source texture to sample, in PIE this is a subsection, in standlone it should be the entire texture
-	 * @param DstRect				Rect within output to write out, since this is within the renderthread should almost always be Dst's Extent.
-	 * @param InRenderingPolicy		RenderingPolicy used to assist / perform processing
-	 */
+	/** Called on the render thread to run a post processing operation on the input texture and produce the output texture. */
+	virtual void PostProcess_Renderthread(FRDGBuilder& GraphBuilder, const FScreenPassTexture& InputTexture, const FScreenPassTexture& OutputTexture)
+	{
+	}
+
+	UE_DEPRECATED(5.5, "Use RDG version instead.")
+PRAGMA_DISABLE_DEPRECATION_WARNINGS
 	virtual void PostProcess_Renderthread(FRHICommandListImmediate& RHICmdList, FRHITexture* Src, FRHITexture* Dst, FIntRect SrcRect, FIntRect DstRect, FSlateRHIRenderingPolicyInterface InRenderingPolicy)
 	{
 	}
+PRAGMA_ENABLE_DEPRECATION_WARNINGS
 
 	/** 
 	 * Called when an post buffer update element is added to a renderbatch, 
@@ -96,9 +91,12 @@ public:
 	 * @param InRenderingPolicy		Slate RHI RenderingPolicy
 	 * @param InSlatePostBuffer		Texture render target used for final output
 	 */
+	UE_DEPRECATED(5.5, "This path is longer supported. Get the render proxy instead.")
+PRAGMA_DISABLE_DEPRECATION_WARNINGS
 	virtual void PostProcess(FRenderResource* InViewInfo, FRenderResource* InViewportTexture, FVector2D InElementWindowSize, FSlateRHIRenderingPolicyInterface InRenderingPolicy, UTextureRenderTarget2D* InSlatePostBuffer)
 	{
 	}
+PRAGMA_ENABLE_DEPRECATION_WARNINGS
 
 	/**
 	 * Gets proxy for this post buffer processor, for execution on the renderthread
@@ -107,39 +105,4 @@ public:
 	{
 		return nullptr;
 	}
-
-protected:
-
-	/** 
-	 * Gets scene backbuffer, typically used as 'Src' texture for post process, but not always (Ex: PIE).
-	 * 
-	 * @param InViewInfo			'FViewportInfo' resource used to get backbuffer in standalone
-	 * @param InViewportTexture		'FSlateRenderTargetRHI' resource used to get the 'BufferedRT' viewport texture used in PIE
-	 * @param InElementWindowSize	Size of window being rendered, used to determine if using stereo rendering or not.
-	 * @param InRHICmdList			RHI command list to queue commands on
-	 */
-	static FTextureRHIRef GetBackbuffer_RenderThread(FRenderResource* InViewInfo, FRenderResource* InViewportTexture, FVector2D InElementWindowSize, FRHICommandListImmediate& InRHICmdList);
-
-	/**
-	 * Gets 'Src' texture for post process command. Typically the scenebuffer.
-	 *
-	 * @param InBackBuffer			Backbuffer used in standalone
-	 * @param InViewportTexture		'FSlateRenderTargetRHI' resource used for 'BufferedRT' viewport texture in PIE
-	 */
-	static FTextureRHIRef GetSrcTexture_RenderThread(FTextureRHIRef InBackBuffer, FRenderResource* InViewportTexture);
-
-	/**
-	 * Gets 'Dst' texture for post process command. Convience method, this should be possible through the direct resource.
-	 *
-	 * @param InSlatePostBuffer		Texture render target used for final output
-	 */
-	static FTextureReferenceRHIRef& GetDstTexture_RenderThread(UTextureRenderTarget2D* InSlatePostBuffer);
-
-	/**
-	 * Gets 'Dst' extent. Used for final size in post process command.
-	 *
-	 * @param InBackBuffer			Backbuffer used for size in standalone
-	 * @param InViewportTexture		'FSlateRenderTargetRHI' resource used for 'BufferedRT' viewport texture size in PIE
-	 */
-	static FIntPoint GetDstExtent_RenderThread(FTextureRHIRef InBackBuffer, FRenderResource* InViewportTexture);
 };
\ No newline at end of file
diff --git a/Engine/Source/Runtime/SlateRHIRenderer/Public/Interfaces/ISlate3DRenderer.h b/Engine/Source/Runtime/SlateRHIRenderer/Public/Interfaces/ISlate3DRenderer.h
--- a/Engine/Source/Runtime/SlateRHIRenderer/Public/Interfaces/ISlate3DRenderer.h
+++ b/Engine/Source/Runtime/SlateRHIRenderer/Public/Interfaces/ISlate3DRenderer.h
@@ -5,6 +5,7 @@
 #include "CoreMinimal.h"
 
 class FRHICommandListImmediate;
+class FRDGBuilder;
 class FSlateDrawBuffer;
 class FTextureRenderTarget2DResource;
 
@@ -72,11 +73,14 @@ public:
 	/** 
 	 * Renders the batched draw elements of the draw buffer to the given render target.
 	 * Call after preparing the draw buffer and render target on the game thread.
-	 * 
+	 *
 	 * @param RenderTarget The render target to render the contents of the draw buffer to
 	 * @param InDrawBuffer The draw buffer containing the batched elements to render
 	 */
-	virtual void DrawWindowToTarget_RenderThread(FRHICommandListImmediate& RHICmdList, const struct FRenderThreadUpdateContext& Context) = 0;
+	virtual void DrawWindowToTarget_RenderThread(FRDGBuilder& GraphBuilder, const struct FRenderThreadUpdateContext& Context) = 0;
+
+	UE_DEPRECATED(5.5, "Use the RDG version instead.")
+	virtual void DrawWindowToTarget_RenderThread(FRHICommandListImmediate& RHICmdList, const struct FRenderThreadUpdateContext& Context) {}
 };
 
 typedef TSharedPtr<ISlate3DRenderer, ESPMode::ThreadSafe> ISlate3DRendererPtr;
diff --git a/Engine/Source/Runtime/SlateRHIRenderer/Public/Interfaces/SlateRHIRenderingPolicyInterface.h b/Engine/Source/Runtime/SlateRHIRenderer/Public/Interfaces/SlateRHIRenderingPolicyInterface.h
--- a/Engine/Source/Runtime/SlateRHIRenderer/Public/Interfaces/SlateRHIRenderingPolicyInterface.h
+++ b/Engine/Source/Runtime/SlateRHIRenderer/Public/Interfaces/SlateRHIRenderingPolicyInterface.h
@@ -2,6 +2,20 @@
 
 #pragma once
 
+#include "ScreenPass.h"
+
+struct FSlatePostProcessSimpleBlurPassInputs
+{
+	FScreenPassTexture InputTexture;
+	FScreenPassTexture OutputTexture;
+	float Strength = 0.0f;
+};
+
+SLATERHIRENDERER_API void AddSlatePostProcessBlurPass(FRDGBuilder& GraphBuilder, const FSlatePostProcessSimpleBlurPassInputs& Inputs);
+
+//////////////////////////////////////////////////////////////////////////
+// Deprecated API
+
 #include "CoreMinimal.h"
 #include "RendererInterface.h"
 #include "Rendering/RenderingCommon.h"
@@ -9,53 +23,26 @@
 #include "Rendering/DrawElements.h"
 #include "Rendering/RenderingPolicy.h"
 
-class FSlateRHIRenderingPolicy;
-
-/**
- * Class used to expose only limited parts of the FSlateRenderingPolicy. 
- * Not an interface to be implemented.
- */
-class SLATERHIRENDERER_API FSlateRHIRenderingPolicyInterface
+class UE_DEPRECATED(5.5, "FSlateRHIRenderingPolicyInterface is no longer used.") FSlateRHIRenderingPolicyInterface
 {
 public:
-	FSlateRHIRenderingPolicyInterface(FSlateRHIRenderingPolicy* InRenderingPolicy);
-
-	static int32 GetProcessSlatePostBuffers();
-
-	bool IsValid() const;
-	bool IsVertexColorInLinearSpace() const;
-	bool GetApplyColorDeficiencyCorrection() const;
+	FSlateRHIRenderingPolicyInterface(const class FSlateRHIRenderingPolicy* InRenderingPolicy) {}
 
-	void BlurRectExternal(FRHICommandListImmediate& RHICmdList, FRHITexture* BlurSrc, FRHITexture* BlurDst, FIntRect SrcRect, FIntRect DstRect, float BlurStrength) const;
+	static int32 GetProcessSlatePostBuffers() { return 0; }
 
-private:
-
-	/** Rendering policy we are an interface to */
-	FSlateRHIRenderingPolicy* RenderingPolicy;
+	bool IsValid() const { return false; }
+	bool IsVertexColorInLinearSpace() const { return false; }
+	bool GetApplyColorDeficiencyCorrection() const { return false; }
+	void BlurRectExternal(FRHICommandListImmediate& RHICmdList, FRHITexture* BlurSrc, FRHITexture* BlurDst, FIntRect SrcRect, FIntRect DstRect, float BlurStrength) const {}
 };
 
-/**
- * RHI version of custom slate element that can accept additional RHI params when performing draw
- */
-class ICustomSlateElementRHI : public ICustomSlateElement
+class UE_DEPRECATED(5.5, "Use ICustomSlateElement instead.") ICustomSlateElementRHI : public ICustomSlateElement
 {
 public:
+	UE_DEPRECATED(5.5, "Use Draw_RenderThread with an RDG builder instead")
+PRAGMA_DISABLE_DEPRECATION_WARNINGS
+	virtual void Draw_RHIRenderThread(FRHICommandListImmediate& RHICmdList, const FTextureRHIRef& RenderTarget, const FSlateCustomDrawParams& Params, FSlateRHIRenderingPolicyInterface RenderingPolicyInterface) {}
+PRAGMA_ENABLE_DEPRECATION_WARNINGS
+};
 
-	/**
-	 * Called from the rendering thread when it is time to render the element
-	 *
-	 * @param RenderTarget				handle to the platform specific render target implementation.  Note this is already bound by Slate initially
-	 * @param Params					Params about current draw state
-	 * @param RenderingPolicyInterface	Interface to current rendering policy
-	 */
-	virtual void Draw_RHIRenderThread(class FRHICommandListImmediate& RHICmdList, const FTextureRHIRef& RenderTarget, const FSlateCustomDrawParams& Params, FSlateRHIRenderingPolicyInterface RenderingPolicyInterface)
-	{
-	}
-
-	//~ Begin ICustomSlateElement interface
-	virtual bool UsesAdditionalRHIParams() const override
-	{
-		return true;
-	}
-	//~ End ICustomSlateElement interface
-};
\ No newline at end of file
+//////////////////////////////////////////////////////////////////////////
\ No newline at end of file
diff --git a/Engine/Source/Runtime/UMG/Private/Slate/SPostBufferUpdate.cpp b/Engine/Source/Runtime/UMG/Private/Slate/SPostBufferUpdate.cpp
--- a/Engine/Source/Runtime/UMG/Private/Slate/SPostBufferUpdate.cpp
+++ b/Engine/Source/Runtime/UMG/Private/Slate/SPostBufferUpdate.cpp
@@ -24,14 +24,14 @@
 /**
  * Custom Slate drawer to update slate post buffer
  */
-class FPostBufferUpdater : public ICustomSlateElementRHI
+class FPostBufferUpdater : public ICustomSlateElement
 {
 public:
 
-	//~ Begin ICustomSlateElementRHI interface
-	virtual void Draw_RHIRenderThread(FRHICommandListImmediate& RHICmdList, const FTextureRHIRef& InWindowBackBuffer, const FSlateCustomDrawParams& Params, class FSlateRHIRenderingPolicyInterface RenderingPolicyInterface) override;
+	//~ Begin ICustomSlateElement interface
+	virtual void Draw_RenderThread(FRDGBuilder& GraphBuilder, const FDrawPassInputs& Inputs) override;
 	virtual void PostCustomElementAdded(class FSlateElementBatcher& ElementBatcher) const override;
-	//~ End ICustomSlateElementRHI interface
+	//~ End ICustomSlateElement interface
 
 public:
 
@@ -57,75 +57,42 @@ public:
 /////////////////////////////////////////////////////
 // FPostBufferUpdater
 
-void FPostBufferUpdater::Draw_RHIRenderThread(FRHICommandListImmediate& RHICmdList, const FTextureRHIRef& InWindowBackBuffer, const FSlateCustomDrawParams& Params, FSlateRHIRenderingPolicyInterface RenderingPolicyInterface)
+void FPostBufferUpdater::Draw_RenderThread(FRDGBuilder& GraphBuilder, const FDrawPassInputs& Inputs)
 {
-	if (FSlateRHIRenderingPolicyInterface::GetProcessSlatePostBuffers())
+	for (ESlatePostRT SlatePostBufferBit : MakeFlagsRange(Inputs.UsedSlatePostBuffers & BuffersToUpdate_Renderthread))
 	{
-		for (ESlatePostRT SlatePostBufferBit : TEnumRange<ESlatePostRT>())
+		UTextureRenderTarget2D* SlatePostBuffer = Cast<UTextureRenderTarget2D>(USlateRHIRendererSettings::Get()->TryGetPostBufferRT(SlatePostBufferBit));
+		if (!SlatePostBuffer)
 		{
-			bool bPostBufferBitUsed = (Params.UsedSlatePostBuffers & SlatePostBufferBit & BuffersToUpdate_Renderthread) != ESlatePostRT::None;
-
-			if (bPostBufferBitUsed)
-			{
-				UTextureRenderTarget2D* SlatePostBuffer = Cast<UTextureRenderTarget2D>(USlateRHIRendererSettings::Get()->TryGetPostBufferRT(SlatePostBufferBit));
-				if (!SlatePostBuffer)
-				{
-					continue;
-				}
-
-				FRHITexture* Src = InWindowBackBuffer.IsValid() ? InWindowBackBuffer.GetReference() : nullptr;
-				FRHITexture* Dst = SlatePostBuffer->TextureReference.TextureReferenceRHI;
+			continue;
+		}
 
-				if (!Src || !Dst)
-				{
-					continue;
-				}
+		// Provided output texture is actually the input into our custom post process texture.
+		const FScreenPassTexture InputTexture(Inputs.OutputTexture, Inputs.SceneViewRect);
 
-				FIntRect SrcRect = FIntRect(0, 0, Src->GetDesc().Extent.X, Src->GetDesc().Extent.Y);
-				FIntRect DstRect = FIntRect(0, 0, Dst->GetDesc().Extent.X, Dst->GetDesc().Extent.Y);
+		// Write to our custom post process texture.
+		const FScreenPassTexture OutputTexture(RegisterExternalTexture(GraphBuilder, SlatePostBuffer->TextureReference.TextureReferenceRHI, TEXT("SlatePostProcessTexture")));
 
-				if (GIsEditor)
-				{
-					// In PIE, the backbuffer for slate will be the entire editor window, So instead use the ViewRect as our extent
-					SrcRect = Params.ViewRect;
-				}
-				
-				// Note: In PIE we draw the scene even during resizes, the ViewRect can change even after a previous draw
-				// leading to crashes if we do not ensure size, we skip below if sizes do not match. This can result in a PostRT 
-				// with no update in PIE (May appear white during active drag-resizing). However this is not an issue in 
-				// Standalone since we don't draw during active resizes.
-				if (SrcRect.Width() != DstRect.Width() || SrcRect.Height() != DstRect.Height())
+		if (TSharedPtr<FSlateRHIPostBufferProcessorProxy> PostProcessorProxy = USlateFXSubsystem::GetPostProcessorProxy(SlatePostBufferBit))
+		{
+			if (TSharedPtr<FSlatePostProcessorUpdaterProxy>* ProcessorUpdaterItr = ProcessorUpdaters.Find(SlatePostBufferBit))
+			{
+				if (TSharedPtr<FSlatePostProcessorUpdaterProxy> ProcessorUpdater = *ProcessorUpdaterItr)
 				{
-					continue;
-				}
+					ProcessorUpdater->UpdateProcessor_RenderThread(PostProcessorProxy);
 
-				if (TSharedPtr<FSlateRHIPostBufferProcessorProxy> PostProcessorProxy = USlateFXSubsystem::GetPostProcessorProxy(SlatePostBufferBit))
-				{
-					if (TSharedPtr<FSlatePostProcessorUpdaterProxy>* ProcessorUpdaterItr = ProcessorUpdaters.Find(SlatePostBufferBit))
+					if (ProcessorUpdater->bSkipBufferUpdate)
 					{
-						if (TSharedPtr<FSlatePostProcessorUpdaterProxy> ProcessorUpdater = *ProcessorUpdaterItr)
-						{
-							ProcessorUpdater->UpdateProcessor_RenderThread(PostProcessorProxy);
-
-							if (ProcessorUpdater->bSkipBufferUpdate)
-							{
-								return;
-							}
-						}
+						return;
 					}
-
-					PostProcessorProxy->PostProcess_Renderthread(RHICmdList, Src, Dst, SrcRect, DstRect, RenderingPolicyInterface);
-				}
-				else
-				{
-					FRHICopyTextureInfo CopyInfo;
-
-					// Copy just the viewport RT if in PIE, else do entire backbuffer
-					CopyInfo.SourcePosition = FIntVector(SrcRect.Min.X, SrcRect.Min.Y, 0);
-					CopyInfo.Size = FIntVector(DstRect.Width(), DstRect.Height(), 1);
-					TransitionAndCopyTexture(RHICmdList, Src, Dst, CopyInfo);
 				}
 			}
+
+			PostProcessorProxy->PostProcess_Renderthread(GraphBuilder, InputTexture, OutputTexture);
+		}
+		else
+		{
+			AddDrawTexturePass(GraphBuilder, FScreenPassViewInfo(), InputTexture, OutputTexture);
 		}
 	}
 }
@@ -142,24 +109,16 @@ void FPostBufferUpdater::PostCustomElementAdded(FSlateElementBatcher& ElementBat
 	}
 
 	// Give proxies a chance to update their renderthread values.
-	if (FSlateRHIRenderingPolicyInterface::GetProcessSlatePostBuffers())
+	for (ESlatePostRT SlatePostBufferBit : MakeFlagsRange(BuffersToUpdate_Renderthread))
 	{
-		for (ESlatePostRT SlatePostBufferBit : TEnumRange<ESlatePostRT>())
+		if (!USlateRHIRendererSettings::Get()->GetSlatePostSetting(SlatePostBufferBit).bEnabled)
 		{
-			if (!USlateRHIRendererSettings::Get()->GetSlatePostSetting(SlatePostBufferBit).bEnabled)
-			{
-				continue;
-			}
-
-			bool bPostBufferBitUsed = (SlatePostBufferBit & BuffersToUpdate_Renderthread) != ESlatePostRT::None;
+			continue;
+		}
 
-			if (bPostBufferBitUsed)
-			{
-				if (TSharedPtr<FSlateRHIPostBufferProcessorProxy> PostProcessorProxy = USlateFXSubsystem::GetPostProcessorProxy(SlatePostBufferBit))
-				{
-					PostProcessorProxy->OnUpdateValuesRenderThread();
-				}
-			}
+		if (TSharedPtr<FSlateRHIPostBufferProcessorProxy> PostProcessorProxy = USlateFXSubsystem::GetPostProcessorProxy(SlatePostBufferBit))
+		{
+			PostProcessorProxy->OnUpdateValuesRenderThread();
 		}
 	}
 }
diff --git a/Engine/Source/Runtime/UMG/UMG.Build.cs b/Engine/Source/Runtime/UMG/UMG.Build.cs
--- a/Engine/Source/Runtime/UMG/UMG.Build.cs
+++ b/Engine/Source/Runtime/UMG/UMG.Build.cs
@@ -17,6 +17,7 @@ public class UMG : ModuleRules
 				"Slate",
 				"SlateCore",
 				"RenderCore",
+				"Renderer",
 				"RHI",
 				"ApplicationCore"
 			}
