Fixed build breakages.

diff --git a/Engine/Source/Runtime/SlateRHIRenderer/Private/SlateRHIRenderer.cpp b/Engine/Source/Runtime/SlateRHIRenderer/Private/SlateRHIRenderer.cpp
--- a/Engine/Source/Runtime/SlateRHIRenderer/Private/SlateRHIRenderer.cpp
+++ b/Engine/Source/Runtime/SlateRHIRenderer/Private/SlateRHIRenderer.cpp
@@ -585,6 +585,12 @@ public:
 	class FApplyColorDeficiency : SHADER_PERMUTATION_BOOL("APPLY_COLOR_DEFICIENCY");
 	using FPermutationDomain = TShaderPermutationDomain<FSCRGBEncoding, FApplyColorDeficiency>;
 
+	FCompositeShader() {}
+
+	FCompositeShader(const ShaderMetaType::CompiledShaderInitializerType& Initializer)
+		: FGlobalShader(Initializer)
+	{}
+
 	static bool ShouldCompilePermutation(const FGlobalShaderPermutationParameters& Parameters)
 	{
 		return SupportsCompositeUIWithSceneHDR(Parameters.Platform);
diff --git a/Engine/Source/Runtime/SlateRHIRenderer/Private/SlateRHIRenderer.h b/Engine/Source/Runtime/SlateRHIRenderer/Private/SlateRHIRenderer.h
--- a/Engine/Source/Runtime/SlateRHIRenderer/Private/SlateRHIRenderer.h
+++ b/Engine/Source/Runtime/SlateRHIRenderer/Private/SlateRHIRenderer.h
@@ -131,7 +131,7 @@ private:
 	 */
 	void CleanUpdatableTextures();
 
-	virtual void OnVirtualDesktopSizeChanged(const FDisplayMetrics& NewDisplayMetric);
+	void OnVirtualDesktopSizeChanged(const FDisplayMetrics& NewDisplayMetric);
 
 	/** A mapping of SWindows to their RHI implementation */
 	TMap<const SWindow*, FSlateViewportInfo*> WindowToViewportInfo;
diff --git a/Engine/Source/Runtime/SlateRHIRenderer/Private/SlateShaders.h b/Engine/Source/Runtime/SlateRHIRenderer/Private/SlateShaders.h
--- a/Engine/Source/Runtime/SlateRHIRenderer/Private/SlateShaders.h
+++ b/Engine/Source/Runtime/SlateRHIRenderer/Private/SlateShaders.h
@@ -330,6 +330,10 @@ class FSlateMaskingPS : public FGlobalShader
 {
 public:
 	DECLARE_GLOBAL_SHADER(FSlateMaskingPS);
+	SHADER_USE_PARAMETER_STRUCT(FSlateMaskingPS, FGlobalShader);
+
+	BEGIN_SHADER_PARAMETER_STRUCT(FParameters, )
+	END_SHADER_PARAMETER_STRUCT()
 };
 
 /** The simple element vertex declaration. */
