Fix Android build break on Slate RDG changes

diff --git a/Engine/Source/Runtime/SlateRHIRenderer/Private/SlateRHIRenderer.cpp b/Engine/Source/Runtime/SlateRHIRenderer/Private/SlateRHIRenderer.cpp
--- a/Engine/Source/Runtime/SlateRHIRenderer/Private/SlateRHIRenderer.cpp
+++ b/Engine/Source/Runtime/SlateRHIRenderer/Private/SlateRHIRenderer.cpp
@@ -664,8 +664,12 @@ void FSlateRHIRenderer::DrawWindow_RenderThread(FRDGBuilder& GraphBuilder, const
 
 	{
 		RDG_GPU_MASK_SCOPE(GraphBuilder, FRHIGPUMask::FromIndex(RHIGetViewportNextPresentGPUIndex(ViewportInfo.ViewportRHI)));
+#if WANTS_DRAW_MESH_EVENTS
 		RDG_EVENT_SCOPE_CONDITIONAL(GraphBuilder, Inputs.WindowTitle.IsEmpty(), "SlateUI Title = <none>");
 		RDG_EVENT_SCOPE_CONDITIONAL(GraphBuilder, !Inputs.WindowTitle.IsEmpty(),"SlateUI Title = %s", *Inputs.WindowTitle);
+#else
+		RDG_EVENT_SCOPE(GraphBuilder, "SlateUI");
+#endif
 		RDG_GPU_STAT_SCOPE(GraphBuilder, SlateUI);
 		RDG_CSV_STAT_EXCLUSIVE_SCOPE(GraphBuilder, Slate);
 		TRACE_CPUPROFILER_EVENT_SCOPE(Slate::DrawWindow_RenderThread);
