Fix nullptr crash on blur widgets

diff --git a/Engine/Source/Runtime/SlateRHIRenderer/Private/SlateRHIRenderingPolicy.cpp b/Engine/Source/Runtime/SlateRHIRenderer/Private/SlateRHIRenderingPolicy.cpp
--- a/Engine/Source/Runtime/SlateRHIRenderer/Private/SlateRHIRenderingPolicy.cpp
+++ b/Engine/Source/Runtime/SlateRHIRenderer/Private/SlateRHIRenderingPolicy.cpp
@@ -1459,6 +1459,11 @@ void AddSlateDrawElementsPass(
 
 			FSlatePostProcessBlurPassInputs BlurInputs;
 
+			if (!Inputs.SceneViewportTexture)
+			{
+				break;
+			}
+
 			if (Inputs.SceneViewportTexture != Inputs.ElementsTexture)
 			{
 				// Blur uses the scene viewport texture output as blur input and composites UI separately.
