Fixed blur to fall back to using the elements texture if no scene texture is provided.

diff --git a/Engine/Source/Runtime/SlateRHIRenderer/Private/SlateRHIRenderingPolicy.cpp b/Engine/Source/Runtime/SlateRHIRenderer/Private/SlateRHIRenderingPolicy.cpp
--- a/Engine/Source/Runtime/SlateRHIRenderer/Private/SlateRHIRenderingPolicy.cpp
+++ b/Engine/Source/Runtime/SlateRHIRenderer/Private/SlateRHIRenderingPolicy.cpp
@@ -1459,12 +1459,7 @@ void AddSlateDrawElementsPass(
 
 			FSlatePostProcessBlurPassInputs BlurInputs;
 
-			if (!Inputs.SceneViewportTexture)
-			{
-				break;
-			}
-
-			if (Inputs.SceneViewportTexture != Inputs.ElementsTexture)
+			if (Inputs.SceneViewportTexture && Inputs.SceneViewportTexture != Inputs.ElementsTexture)
 			{
 				// Blur uses the scene viewport texture output as blur input and composites UI separately.
 				BlurInputs.InputTexture = Inputs.SceneViewportTexture;
@@ -1481,7 +1476,7 @@ void AddSlateDrawElementsPass(
 			}
 
 			BlurInputs.InputRect = FIntRect(ShaderParams.PixelParams.X, ShaderParams.PixelParams.Y, ShaderParams.PixelParams.Z, ShaderParams.PixelParams.W);
-			BlurInputs.OutputTexture = Inputs.SceneViewportTexture;
+			BlurInputs.OutputTexture = Inputs.SceneViewportTexture ? Inputs.SceneViewportTexture : Inputs.ElementsTexture;
 			BlurInputs.OutputRect = BlurInputs.InputRect;
 			BlurInputs.OutputLoadAction = ConsumeLoadAction(ElementsLoadAction);
 			BlurInputs.ClippingOp = NextClippingOp;
