Refactored Slate to use shader bindings to make the RDG execution lambda safe for async execution.

diff --git a/Engine/Shaders/Private/SlateElementPixelShader.usf b/Engine/Shaders/Private/SlateElementPixelShader.usf
--- a/Engine/Shaders/Private/SlateElementPixelShader.usf
+++ b/Engine/Shaders/Private/SlateElementPixelShader.usf
@@ -35,7 +35,8 @@ SamplerState ElementTextureSampler;
 
 #if SAMPLE_VIRTUAL_TEXTURE
 Texture2D<uint4> InPageTableTexture;
-uint4 VTPackedPageTableUniform[2];
+uint4 VTPackedPageTableUniform0;
+uint4 VTPackedPageTableUniform1;
 uint4 VTPackedUniform;
 float4 VTShaderParams;
 #endif // SAMPLE_VIRTUAL_TEXTURE
@@ -221,7 +222,7 @@ half4 GetTextureColor( VertexToPixelInterpolants InVertex, float2 UV )
 	half4 BaseColor;
 
 	#if SAMPLE_VIRTUAL_TEXTURE
-			VTPageTableResult PageTableResult = TextureLoadVirtualPageTableLevel(InPageTableTexture, VTPageTableUniform_Unpack(VTPackedPageTableUniform[0], VTPackedPageTableUniform[1]), UV, VTADDRESSMODE_WRAP, VTADDRESSMODE_WRAP, VTShaderParams.x); 
+			VTPageTableResult PageTableResult = TextureLoadVirtualPageTableLevel(InPageTableTexture, VTPageTableUniform_Unpack(VTPackedPageTableUniform0, VTPackedPageTableUniform1), UV, VTADDRESSMODE_WRAP, VTADDRESSMODE_WRAP, VTShaderParams.x); 
 			BaseColor = TextureVirtualSampleLevel(ElementTexture, ElementTextureSampler, PageTableResult, (uint)VTShaderParams.y, VTUniform_Unpack(VTPackedUniform));
 	#else
 
diff --git a/Engine/Shaders/Private/SlateVertexShader.usf b/Engine/Shaders/Private/SlateVertexShader.usf
--- a/Engine/Shaders/Private/SlateVertexShader.usf
+++ b/Engine/Shaders/Private/SlateVertexShader.usf
@@ -4,8 +4,6 @@
 #include "GammaCorrectionCommon.ush"
 #include "SlateShaderCommon.ush"
 
-float4x4 ViewProjection;
-
 VertexToPixelInterpolants Main(
 	in float4 InTextureCoordinates : ATTRIBUTE0,
 	in float2 InMaterialTextureCoordinates : ATTRIBUTE1,
@@ -127,7 +125,7 @@ VertexToPixelInterpolants Main(
 
 #endif // USE_MATERIALS
 
-	VOut.Position = mul(WorldPosition, ViewProjection);
+	VOut.Position = mul(WorldPosition, SlateView.ViewProjection);
 
 	return VOut;
 }
diff --git a/Engine/Source/Runtime/SlateRHIRenderer/Private/SlateMaterialShader.cpp b/Engine/Source/Runtime/SlateRHIRenderer/Private/SlateMaterialShader.cpp
--- a/Engine/Source/Runtime/SlateRHIRenderer/Private/SlateMaterialShader.cpp
+++ b/Engine/Source/Runtime/SlateRHIRenderer/Private/SlateMaterialShader.cpp
@@ -5,16 +5,15 @@
 #include "Materials/Material.h"
 #include "RHIStaticStates.h"
 #include "ShaderParameterUtils.h"
+#include "MeshDrawShaderBindings.h"
+#include "SceneInterface.h"
 
 IMPLEMENT_TYPE_LAYOUT(FSlateMaterialShaderVS);
 IMPLEMENT_TYPE_LAYOUT(FSlateMaterialShaderPS);
 
 FSlateMaterialShaderVS::FSlateMaterialShaderVS(const FMaterialShaderType::CompiledShaderInitializerType& Initializer)
 	: FMaterialShader(Initializer)
-{
-	ViewProjection.Bind(Initializer.ParameterMap, TEXT("ViewProjection"));
-}
-
+{}
 
 void FSlateMaterialShaderVS::ModifyCompilationEnvironment(const FMaterialShaderPermutationParameters& Parameters, FShaderCompilerEnvironment& OutEnvironment)
 {
@@ -31,22 +30,18 @@ bool FSlateMaterialShaderVS::ShouldCompilePermutation(const FMaterialShaderPermu
 	return Parameters.MaterialParameters.MaterialDomain == MD_UI;
 }
 
-void FSlateMaterialShaderVS::SetViewProjection(FRHIBatchedShaderParameters& BatchedParameters, const FMatrix44f& InViewProjection )
-{
-	SetShaderValue(BatchedParameters, ViewProjection, InViewProjection );
-}
-
 void FSlateMaterialShaderVS::SetMaterialShaderParameters(
-	FRHIBatchedShaderParameters& BatchedParameters,
+	FMeshDrawSingleShaderBindings& ShaderBindings,
 	const FSceneInterface* Scene,
 	const TUniformBufferRef<FViewUniformShaderParameters>& ViewUniformBuffer,
 	const FMaterialRenderProxy* MaterialRenderProxy,
 	const FMaterial* Material)
 {
 	const auto& ViewUniformBufferParameter = GetUniformBufferParameter<FViewUniformShaderParameters>();
-	SetUniformBufferParameter(BatchedParameters, ViewUniformBufferParameter, ViewUniformBuffer);
+	ShaderBindings.Add(ViewUniformBufferParameter, ViewUniformBuffer);
 
-	FMaterialShader::SetParameters(BatchedParameters, MaterialRenderProxy, *Material, Scene);
+	const ERHIFeatureLevel::Type FeatureLevel = Scene ? Scene->GetFeatureLevel() : GMaxRHIFeatureLevel;
+	FMaterialShader::GetShaderBindings(Scene, FeatureLevel, *MaterialRenderProxy, *Material, ShaderBindings);
 }
 
 bool FSlateMaterialShaderPS::ShouldCompilePermutation(const FMaterialShaderPermutationParameters& Parameters)
@@ -75,79 +70,43 @@ FSlateMaterialShaderPS::FSlateMaterialShaderPS(const FMaterialShaderType::Compil
 	TextureParameterSampler.Bind(Initializer.ParameterMap, TEXT("ElementTextureSampler"));
 }
 
-void FSlateMaterialShaderPS::SetBlendState(FGraphicsPipelineStateInitializer& GraphicsPSOInit, const FMaterial* Material)
-{
-	EBlendMode BlendMode = Material->GetBlendMode();
-
-	switch (BlendMode)
-	{
-	default:
-	case BLEND_Opaque:
-		GraphicsPSOInit.BlendState = TStaticBlendState<>::GetRHI();
-		break;
-	case BLEND_Masked:
-		GraphicsPSOInit.BlendState = TStaticBlendState<>::GetRHI();
-		break;
-	case BLEND_Translucent:
-		GraphicsPSOInit.BlendState = TStaticBlendState<CW_RGBA, BO_Add, BF_SourceAlpha, BF_InverseSourceAlpha, BO_Add, BF_InverseDestAlpha, BF_One>::GetRHI();
-		break;
-	case BLEND_Additive:
-		// Add to the existing scene color
-		GraphicsPSOInit.BlendState = TStaticBlendState<CW_RGBA, BO_Add, BF_One, BF_One, BO_Add, BF_One, BF_One>::GetRHI();
-		break;
-	case BLEND_Modulate:
-		// Modulate with the existing scene color
-		GraphicsPSOInit.BlendState = TStaticBlendState<CW_RGB, BO_Add, BF_Zero, BF_SourceColor>::GetRHI();
-		break;
-	case BLEND_AlphaComposite:
-		// Blend with existing scene color. New color is already pre-multiplied by alpha.
-		GraphicsPSOInit.BlendState = TStaticBlendState<CW_RGBA, BO_Add, BF_One, BF_InverseSourceAlpha, BO_Add, BF_One, BF_InverseSourceAlpha>::GetRHI();
-		break;
-	case BLEND_AlphaHoldout:
-		// Blend by holding out the matte shape of the source alpha
-		GraphicsPSOInit.BlendState = TStaticBlendState<CW_RGBA, BO_Add, BF_Zero, BF_InverseSourceAlpha, BO_Add, BF_Zero, BF_InverseSourceAlpha>::GetRHI();
-		break;
-
-	};
-}
-
-void FSlateMaterialShaderPS::SetParameters(
-	FRHIBatchedShaderParameters& BatchedParameters,
+void FSlateMaterialShaderPS::SetMaterialShaderParameters(
+	FMeshDrawSingleShaderBindings& ShaderBindings,
 	const FSceneInterface* Scene,
 	const TUniformBufferRef<FViewUniformShaderParameters>& ViewUniformBuffer,
 	const FMaterialRenderProxy* MaterialRenderProxy,
 	const FMaterial* Material,
 	const FShaderParams& InShaderParams)
 {
-	SetShaderValue(BatchedParameters, ShaderParams, (FVector4f)InShaderParams.PixelParams);
-	SetShaderValue(BatchedParameters, ShaderParams2, (FVector4f)InShaderParams.PixelParams2);
+	ShaderBindings.Add(ShaderParams,  (FVector4f)InShaderParams.PixelParams);
+	ShaderBindings.Add(ShaderParams2, (FVector4f)InShaderParams.PixelParams2);
 
 	const auto& ViewUniformBufferParameter = GetUniformBufferParameter<FViewUniformShaderParameters>();
-	SetUniformBufferParameter(BatchedParameters, ViewUniformBufferParameter, ViewUniformBuffer);
+	ShaderBindings.Add(ViewUniformBufferParameter, ViewUniformBuffer);
 
-	FMaterialShader::SetParameters(BatchedParameters, MaterialRenderProxy, *Material, Scene);
+	const ERHIFeatureLevel::Type FeatureLevel = Scene ? Scene->GetFeatureLevel() : GMaxRHIFeatureLevel;
+	FMaterialShader::GetShaderBindings(Scene, FeatureLevel, *MaterialRenderProxy, *Material, ShaderBindings);
 }
 
-void FSlateMaterialShaderPS::SetAdditionalTexture(FRHIBatchedShaderParameters& BatchedParameters, FRHITexture* InTexture, const FSamplerStateRHIRef SamplerState )
+void FSlateMaterialShaderPS::SetAdditionalTexture(FMeshDrawSingleShaderBindings& ShaderBindings, FRHITexture* InTexture, const FSamplerStateRHIRef SamplerState)
 {
-	SetTextureParameter(BatchedParameters, AdditionalTextureParameter, TextureParameterSampler, SamplerState, InTexture );
+	ShaderBindings.AddTexture(AdditionalTextureParameter, TextureParameterSampler, SamplerState, InTexture);
 }
 
-void FSlateMaterialShaderPS::SetDisplayGammaAndContrast(FRHIBatchedShaderParameters& BatchedParameters, float InDisplayGamma, float InContrast)
+void FSlateMaterialShaderPS::SetDisplayGammaAndContrast(FMeshDrawSingleShaderBindings& ShaderBindings, float InDisplayGamma, float InContrast)
 {
 	FVector4f InGammaValues(2.2f / InDisplayGamma, 1.0f / InDisplayGamma, 0.0f, InContrast);
 
-	SetShaderValue(BatchedParameters, GammaAndAlphaValues, InGammaValues);
+	ShaderBindings.Add(GammaAndAlphaValues, InGammaValues);
 }
 
-void FSlateMaterialShaderPS::SetDrawFlags(FRHIBatchedShaderParameters& BatchedParameters, bool bDrawDisabledEffect)
+void FSlateMaterialShaderPS::SetDrawFlags(FMeshDrawSingleShaderBindings& ShaderBindings, bool bDrawDisabledEffect)
 {
 	FVector4f InDrawFlags((bDrawDisabledEffect ? 1.f : 0.f), 0.f, 0.f, 0.f);
 
-	SetShaderValue(BatchedParameters, DrawFlags, InDrawFlags);
+	ShaderBindings.Add(DrawFlags, InDrawFlags);
 }
 
-
 #define IMPLEMENT_SLATE_VERTEXMATERIALSHADER_TYPE(bUseInstancing) \
 	typedef TSlateMaterialShaderVS<bUseInstancing> TSlateMaterialShaderVS##bUseInstancing; \
 	IMPLEMENT_MATERIAL_SHADER_TYPE(template<>, TSlateMaterialShaderVS##bUseInstancing, TEXT("/Engine/Private/SlateVertexShader.usf"), TEXT("Main"), SF_Vertex);
diff --git a/Engine/Source/Runtime/SlateRHIRenderer/Private/SlateMaterialShader.h b/Engine/Source/Runtime/SlateRHIRenderer/Private/SlateMaterialShader.h
--- a/Engine/Source/Runtime/SlateRHIRenderer/Private/SlateMaterialShader.h
+++ b/Engine/Source/Runtime/SlateRHIRenderer/Private/SlateMaterialShader.h
@@ -9,6 +9,8 @@
 #include "MaterialShaderType.h"
 #include "MaterialShader.h"
 
+class FMeshDrawSingleShaderBindings;
+
 class FSlateMaterialShaderVS : public FMaterialShader
 {
 	DECLARE_TYPE_LAYOUT(FSlateMaterialShaderVS, NonVirtual);
@@ -20,23 +22,12 @@ public:
 
 	static bool ShouldCompilePermutation(const FMaterialShaderPermutationParameters& Parameters);
 
-	/** 
-	 * Sets the view projection parameter
-	 *
-	 * @param InViewProjection	The ViewProjection matrix to use when this shader is bound 
-	 */
-	void SetViewProjection(FRHIBatchedShaderParameters& BatchedParameters, const FMatrix44f& InViewProjection);
-
 	void SetMaterialShaderParameters(
-		FRHIBatchedShaderParameters& BatchedParameters,
+		FMeshDrawSingleShaderBindings& ShaderBindings,
 		const FSceneInterface* Scene,
 		const TUniformBufferRef<FViewUniformShaderParameters>& ViewUniformBuffer,
 		const FMaterialRenderProxy* MaterialRenderProxy,
 		const FMaterial* Material);
-
-private:
-	/** ViewProjection parameter used by the shader */
-	LAYOUT_FIELD(FShaderParameter, ViewProjection)
 };
 
 class FSlateMaterialShaderPS : public FMaterialShader
@@ -53,21 +44,19 @@ public:
 	FSlateMaterialShaderPS() {}
 	FSlateMaterialShaderPS(const FMaterialShaderType::CompiledShaderInitializerType& Initializer);
 
-	void SetBlendState(FGraphicsPipelineStateInitializer& GraphicsPSOInit, const FMaterial* Material);
-
-	void SetParameters(
-		FRHIBatchedShaderParameters& BatchedParameters,
+	void SetMaterialShaderParameters(
+		FMeshDrawSingleShaderBindings& ShaderBindings,
 		const FSceneInterface* Scene,
 		const TUniformBufferRef<FViewUniformShaderParameters>& ViewUniformBuffer,
 		const FMaterialRenderProxy* MaterialRenderProxy,
 		const FMaterial* Material,
 		const FShaderParams& InShaderParams);
 
-	void SetDisplayGammaAndContrast(FRHIBatchedShaderParameters& BatchedParameters, float InDisplayGamma, float InContrast);
+	void SetDisplayGammaAndContrast(FMeshDrawSingleShaderBindings& ShaderBindings, float InDisplayGamma, float InContrast);
 
-	void SetAdditionalTexture(FRHIBatchedShaderParameters& BatchedParameters, FRHITexture* InTexture, const FSamplerStateRHIRef SamplerState );
+	void SetAdditionalTexture(FMeshDrawSingleShaderBindings& ShaderBindings, FRHITexture* InTexture, const FSamplerStateRHIRef SamplerState );
 
-	void SetDrawFlags(FRHIBatchedShaderParameters& BatchedParameters, bool bDrawDisabledEffect);
+	void SetDrawFlags(FMeshDrawSingleShaderBindings& ShaderBindings, bool bDrawDisabledEffect);
 
 private:
 	LAYOUT_FIELD(FShaderParameter, GammaAndAlphaValues);
diff --git a/Engine/Source/Runtime/SlateRHIRenderer/Private/SlateRHIRenderingPolicy.cpp b/Engine/Source/Runtime/SlateRHIRenderer/Private/SlateRHIRenderingPolicy.cpp
--- a/Engine/Source/Runtime/SlateRHIRenderer/Private/SlateRHIRenderingPolicy.cpp
+++ b/Engine/Source/Runtime/SlateRHIRenderer/Private/SlateRHIRenderingPolicy.cpp
@@ -33,6 +33,7 @@
 #include "SceneRenderTargetParameters.h"
 #include "SceneInterface.h"
 #include "Containers/StaticBitArray.h"
+#include "MeshPassProcessor.h"
 
 extern void UpdateNoiseTextureParameters(FViewUniformShaderParameters& ViewUniformShaderParameters);
 
@@ -468,6 +469,47 @@ EPrimitiveType GetRHIPrimitiveType(ESlateDrawPrimitive SlateType)
 	}
 };
 
+FRHIBlendState* GetMaterialBlendState(FSlateShaderResource* TextureMaskResource, const FMaterial* Material)
+{
+	if (TextureMaskResource && IsOpaqueOrMaskedBlendMode(*Material))
+	{
+		// Font materials require some form of translucent blending
+		return TStaticBlendState<CW_RGBA, BO_Add, BF_SourceAlpha, BF_InverseSourceAlpha, BO_Add, BF_InverseDestAlpha, BF_One>::GetRHI();
+	}
+	else
+	{
+		switch (Material->GetBlendMode())
+		{
+		default:
+		case BLEND_Opaque:
+			return TStaticBlendState<>::GetRHI();
+			break;
+		case BLEND_Masked:
+			return TStaticBlendState<>::GetRHI();
+			break;
+		case BLEND_Translucent:
+			return TStaticBlendState<CW_RGBA, BO_Add, BF_SourceAlpha, BF_InverseSourceAlpha, BO_Add, BF_InverseDestAlpha, BF_One>::GetRHI();
+			break;
+		case BLEND_Additive:
+			// Add to the existing scene color
+			return TStaticBlendState<CW_RGBA, BO_Add, BF_One, BF_One, BO_Add, BF_One, BF_One>::GetRHI();
+			break;
+		case BLEND_Modulate:
+			// Modulate with the existing scene color
+			return TStaticBlendState<CW_RGB, BO_Add, BF_Zero, BF_SourceColor>::GetRHI();
+			break;
+		case BLEND_AlphaComposite:
+			// Blend with existing scene color. New color is already pre-multiplied by alpha.
+			return TStaticBlendState<CW_RGBA, BO_Add, BF_One, BF_InverseSourceAlpha, BO_Add, BF_One, BF_InverseSourceAlpha>::GetRHI();
+			break;
+		case BLEND_AlphaHoldout:
+			// Blend by holding out the matte shape of the source alpha
+			return TStaticBlendState<CW_RGBA, BO_Add, BF_Zero, BF_InverseSourceAlpha, BO_Add, BF_Zero, BF_InverseSourceAlpha>::GetRHI();
+			break;
+		};
+	}
+}
+
 //////////////////////////////////////////////////////////////////////////
 
 FSlateElementsBuffers BuildSlateElementsBuffers(FRDGBuilder& GraphBuilder, FSlateBatchData& BatchData)
@@ -518,6 +560,13 @@ FSlateElementsBuffers BuildSlateElementsBuffers(FRDGBuilder& GraphBuilder, FSlat
 
 //////////////////////////////////////////////////////////////////////////
 
+BEGIN_UNIFORM_BUFFER_STRUCT(FSlateViewUniformParameters, )
+	SHADER_PARAMETER(FMatrix44f, ViewProjection)
+END_UNIFORM_BUFFER_STRUCT()
+
+IMPLEMENT_STATIC_UNIFORM_BUFFER_SLOT(SlateView);
+IMPLEMENT_STATIC_UNIFORM_BUFFER_STRUCT(FSlateViewUniformParameters, "SlateView", SlateView);
+
 struct FSlateSceneViewAllocateInputs
 {
 	FIntPoint TextureExtent             = FIntPoint::ZeroValue;
@@ -538,9 +587,9 @@ class FSlateSceneViewAllocator
 {
 	RDG_FRIEND_ALLOCATOR_FRIEND(FSlateSceneViewAllocator);
 public:
-	static FSlateSceneViewAllocator* Create(FRDGBuilder& GraphBuilder, FSlateRHIResourceManager& ResourceManager)
+	static FSlateSceneViewAllocator* Create(FRDGBuilder& GraphBuilder, FSlateRHIResourceManager& ResourceManager, const FSlateSceneViewAllocateInputs& Inputs)
 	{
-		return GraphBuilder.AllocObject<FSlateSceneViewAllocator>(ResourceManager);
+		return GraphBuilder.AllocObject<FSlateSceneViewAllocator>(ResourceManager, Inputs);
 	}
 
 	const TUniformBufferRef<FViewUniformShaderParameters>& GetViewUniformBuffer(const FSlateSceneView* View) const
@@ -548,7 +597,7 @@ public:
 		return UniformBuffers[(int32)View->FeatureLevel];
 	}
 
-	const FSlateSceneView* BeginAllocateSceneView(FRDGBuilder& GraphBuilder, FSlateRHIResourceManager& ResourceManager, int32 SceneViewIndex, const FSlateSceneViewAllocateInputs& Inputs)
+	const FSlateSceneView* BeginAllocateSceneView(FRDGBuilder& GraphBuilder, int32 SceneViewIndex)
 	{
 		if (!SceneViews.IsValidIndex(SceneViewIndex))
 		{
@@ -560,16 +609,17 @@ public:
 
 		if (TUniformBufferRef<FViewUniformShaderParameters>& UniformBuffer = UniformBuffers[(int32)FeatureLevel]; !UniformBuffer)
 		{
-			UniformBuffer = CreateUniformBuffer(FeatureLevel, Inputs);
+			UniformBuffer = CreateUniformBuffer(FeatureLevel, AllocateInputs);
 		}
 
 		return &SceneView;
 	}
 
 private:
-	FSlateSceneViewAllocator(FSlateRHIResourceManager& ResourceManager)
+	FSlateSceneViewAllocator(FSlateRHIResourceManager& ResourceManager, const FSlateSceneViewAllocateInputs& Inputs)
 		: SceneViewWithNullSceneIndex(ResourceManager.GetSceneCount())
 		, NumScenes(SceneViewWithNullSceneIndex + 1)
+		, AllocateInputs(Inputs)
 	{
 		SceneViews.SetNum(NumScenes);
 		SceneViews.Last().FeatureLevel = GMaxRHIFeatureLevel;
@@ -632,6 +682,7 @@ private:
 
 	const int32 SceneViewWithNullSceneIndex;
 	const int32 NumScenes;
+	const FSlateSceneViewAllocateInputs AllocateInputs;
 	TArray<FSlateSceneView, FRDGArrayAllocator> SceneViews;
 	TStaticArray<TUniformBufferRef<FViewUniformShaderParameters>, (int32)ERHIFeatureLevel::Num> UniformBuffers;
 };
@@ -901,36 +952,54 @@ inline ESlateRenderBatchType GetSlateRenderBatchType(const FSlateRenderBatch& Dr
 	return ESlateRenderBatchType::Primitive;
 }
 
-struct FSlateRenderBatchOp
+class FSlateDrawShaderBindings : public FMeshDrawSingleShaderBindings
 {
-	FSlateRenderBatchOp() = default;
-
-	FSlateRenderBatchOp(const FSlateRenderBatch* InRenderBatch, const FSlateClippingOp* InClippingStateOp, const FSlateSceneView* InSceneView)
-		: RenderBatch(InRenderBatch)
-		, ClippingStateOp(InClippingStateOp)
-		, SceneView(InSceneView)
+	FSlateDrawShaderBindings(const TShaderRef<FShader>& InShader, const FMeshDrawShaderBindingsLayout& InLayout, uint8* InData)
+		: FMeshDrawSingleShaderBindings(InLayout, InData)
+		, Shader(InShader)
 	{}
 
+public:
+	static FSlateDrawShaderBindings* Create(FRDGBuilder& GraphBuilder, const TShaderRef<FShader>& Shader)
+	{
+		const FMeshDrawShaderBindingsLayout Layout(Shader);
+		const uint32 DataSize = Layout.GetDataSizeBytes();
+		uint8* Data = (uint8*)GraphBuilder.Alloc(DataSize);
+		FMemory::Memzero(Data, DataSize);
+		return new (GraphBuilder.Alloc(sizeof(FSlateDrawShaderBindings))) FSlateDrawShaderBindings(Shader, Layout, Data);
+	}
+
+	void SetOnCommandList(FRHICommandList& RHICmdList) const
+	{
+		FRHIBatchedShaderParameters& BatchedParameters = RHICmdList.GetScratchShaderParameters();
+		FReadOnlyMeshDrawSingleShaderBindings::SetShaderBindings(BatchedParameters, FReadOnlyMeshDrawSingleShaderBindings(*this));
+		RHICmdList.SetBatchedShaderParameters(Shader.GetGraphicsShader(), BatchedParameters);
+	}
+
+	TShaderRef<FShader> Shader;
+};
+
+struct FSlateRenderBatchOp
+{
+	FSlateRenderBatchOp* Next;
 	const FSlateRenderBatch* RenderBatch;
 	const FSlateClippingOp* ClippingStateOp;
-	const FSlateSceneView* SceneView = nullptr;
+	FSlateDrawShaderBindings* VertexBindings;
+	FSlateDrawShaderBindings* PixelBindings;
+	FRHIBlendState* BlendState;
+	ESlateShaderResource::Type ShaderResourceType;
 };
 
 //////////////////////////////////////////////////////////////////////////
 
-struct FSlateRenderBatchDrawInputs
+struct FSlateRenderBatchCreateInputs
 {
-	FMatrix44f ElementsMatrix;
-	FVector2f ElementsOffset;
-	FIntRect ElementsViewRect;
-	FIntRect SceneViewRect;
-	const FSlateSceneViewAllocator* SceneViewAllocator;
+	FGlobalShaderMap* ShaderMap;
+	FSlateSceneViewAllocator* SceneViewAllocator;
 	TConstArrayView<FTextureLODGroup> TextureLODGroups;
-	FSlateElementsBuffers ElementsBuffers;
 	float DisplayGamma;
 	float DisplayContrast;
 	float EngineGamma;
-	bool bWireframe;
 #if WITH_SLATE_VISUALIZERS
 	FLinearColor BatchColor;
 #endif
@@ -943,68 +1012,53 @@ struct FSlateRenderBatchDrawState
 	uint8 StencilRef = 0;
 };
 
-void DrawSlateRenderBatch(
-	FRHICommandList& RHICmdList,
-	FSlateRenderBatchDrawState& State,
-	const FSlateRenderBatchDrawInputs& Inputs,
-	const FSlateRenderBatchOp& RenderBatchOp)
+FSlateRenderBatchOp* CreateSlateRenderBatchOp(
+	FRDGBuilder& GraphBuilder,
+	const FSlateRenderBatchCreateInputs& Inputs,
+	const FSlateRenderBatch* RenderBatch,
+	const FSlateClippingOp* ClippingStateOp)
 {
-	const FSlateClippingOp* ClippingStateOp = RenderBatchOp.ClippingStateOp;
-	const FSlateRenderBatch& RenderBatch         = *RenderBatchOp.RenderBatch;
-	
-	FGlobalShaderMap* ShaderMap = GetGlobalShaderMap(GMaxRHIShaderPlatform);
-
-	if (State.LastClippingOp != ClippingStateOp)
-	{
-		GetSlateClippingPipelineState(ClippingStateOp, State.GraphicsPSOInit.DepthStencilState, State.StencilRef);
-		SetSlateClipping(RHICmdList, ClippingStateOp, Inputs.ElementsViewRect);
-		State.LastClippingOp = ClippingStateOp;
-	}
-
-	FRHIBuffer* ElementsVertexBuffer = Inputs.ElementsBuffers.VertexBuffer->GetRHI();
-	FRHIBuffer* ElementsIndexBuffer  = Inputs.ElementsBuffers.IndexBuffer->GetRHI();
-
-	const FSlateShaderResource* ShaderResource   = RenderBatch.ShaderResource;
-	const ESlateBatchDrawFlag DrawFlags          = RenderBatch.DrawFlags;
-	const ESlateDrawEffect DrawEffects           = RenderBatch.DrawEffects;
-	const ESlateShader ShaderType                = RenderBatch.ShaderType;
-	const FShaderParams& ShaderParams            = RenderBatch.ShaderParams;
-
-	if (EnumHasAllFlags(DrawFlags, ESlateBatchDrawFlag::Wireframe))
-	{
-		State.GraphicsPSOInit.RasterizerState = TStaticRasterizerState<FM_Wireframe>::GetRHI();
-	}
-	else
-	{
-		State.GraphicsPSOInit.RasterizerState = TStaticRasterizerState<FM_Solid>::GetRHI();
-	}
-
-	check(RenderBatch.NumIndices > 0);
-	const uint32 PrimitiveCount = RenderBatch.DrawPrimitiveType == ESlateDrawPrimitive::LineList ? RenderBatch.NumIndices / 2 : RenderBatch.NumIndices / 3;
+	const FSlateShaderResource* ShaderResource   = RenderBatch->ShaderResource;
+	const ESlateBatchDrawFlag DrawFlags          = RenderBatch->DrawFlags;
+	const ESlateDrawEffect DrawEffects           = RenderBatch->DrawEffects;
+	const ESlateShader ShaderType                = RenderBatch->ShaderType;
+	const FShaderParams& ShaderParams            = RenderBatch->ShaderParams;
 
 	check(ShaderResource == nullptr || !ShaderResource->Debug_IsDestroyed());
 	const ESlateShaderResource::Type ResourceType = ShaderResource ? ShaderResource->GetType() : ESlateShaderResource::Invalid;
 
-	const bool bUseInstancing = RenderBatch.InstanceCount > 0 && RenderBatch.InstanceData != nullptr;
+	const bool bUseInstancing = RenderBatch->InstanceCount > 0 && RenderBatch->InstanceData != nullptr;
+
+	const float FinalGamma = EnumHasAnyFlags(DrawFlags, ESlateBatchDrawFlag::ReverseGamma) ? (1.0f / Inputs.EngineGamma) : EnumHasAnyFlags(DrawFlags, ESlateBatchDrawFlag::NoGamma) ? 1.0f : Inputs.DisplayGamma;
+	const float FinalContrast = EnumHasAnyFlags(DrawFlags, ESlateBatchDrawFlag::NoGamma) ? 1 : Inputs.DisplayContrast;
+
+	FRHIBlendState* BlendState = nullptr;
+	FSlateDrawShaderBindings* PixelBindings = nullptr;
+	FSlateDrawShaderBindings* VertexBindings = nullptr;
 
 	if (ResourceType == ESlateShaderResource::Material)
 	{
+		// Skip material render batches when the engine is not available.
+		if (!GEngine)
+		{
+			return nullptr;
+		}
+
 		FSlateMaterialResource* MaterialShaderResource = (FSlateMaterialResource*)ShaderResource;
 		MaterialShaderResource->CheckForStaleResources();
 
 		const FMaterialRenderProxy* MaterialRenderProxy = MaterialShaderResource->GetRenderProxy();
+
 		if (!MaterialRenderProxy)
 		{
-			return;
+			return nullptr;
 		}
-		
-		const FSlateSceneView* SceneView = RenderBatchOp.SceneView;
+
+		const FSlateSceneView* SceneView = Inputs.SceneViewAllocator->BeginAllocateSceneView(GraphBuilder, RenderBatch->SceneIndex);
 		const ERHIFeatureLevel::Type SceneFeatureLevel = SceneView->FeatureLevel;
 		const FSceneInterface* Scene = SceneView->Scene;
 		const TUniformBufferRef<FViewUniformShaderParameters>& ViewUniformBuffer = Inputs.SceneViewAllocator->GetViewUniformBuffer(SceneView);
 
-		MaterialRenderProxy->UpdateUniformExpressionCacheIfNeeded(RHICmdList, SceneFeatureLevel);
-
 		TShaderRef<FSlateMaterialShaderVS> VertexShader;
 		TShaderRef<FSlateMaterialShaderPS> PixelShader;
 
@@ -1014,7 +1068,7 @@ void DrawSlateRenderBatch(
 
 		while (MaterialRenderProxy)
 		{
-			const FMaterial* Material = MaterialRenderProxy->GetMaterialNoFallback(SceneFeatureLevel);
+			const FMaterial* Material = MaterialRenderProxy->UpdateUniformExpressionCacheIfNeeded(GraphBuilder.RHICmdList, SceneFeatureLevel);
 			FMaterialShaders Shaders;
 			if (Material && Material->TryGetShaders(ShaderTypesToGet, nullptr, Shaders))
 			{
@@ -1029,70 +1083,27 @@ void DrawSlateRenderBatch(
 
 		if (!VertexShader.IsValid() || !PixelShader.IsValid())
 		{
-			return;
-		}
-
-		SLATE_DRAW_EVENTF(RHICmdList, MaterialBatch, TEXT("Slate Material: %s"), MaterialRenderProxy->GetMaterialName());
-
-		PixelShader->SetBlendState(State.GraphicsPSOInit, EffectiveMaterial);
-
-		FSlateShaderResource* MaskResource = MaterialShaderResource->GetTextureMaskResource();
-
-		if (MaskResource && IsOpaqueOrMaskedBlendMode(*EffectiveMaterial))
-		{
-			// Font materials require some form of translucent blending
-			State.GraphicsPSOInit.BlendState = TStaticBlendState<CW_RGBA, BO_Add, BF_SourceAlpha, BF_InverseSourceAlpha, BO_Add, BF_InverseDestAlpha, BF_One>::GetRHI();
+			return nullptr;
 		}
 
-		State.GraphicsPSOInit.BoundShaderState.VertexDeclarationRHI = bUseInstancing ? GSlateInstancedVertexDeclaration.VertexDeclarationRHI : GSlateVertexDeclaration.VertexDeclarationRHI;
-		State.GraphicsPSOInit.BoundShaderState.VertexShaderRHI = VertexShader.GetVertexShader();
-		State.GraphicsPSOInit.BoundShaderState.PixelShaderRHI = PixelShader.GetPixelShader();
-		State.GraphicsPSOInit.PrimitiveType = GetRHIPrimitiveType(RenderBatch.DrawPrimitiveType);
-
-		SetGraphicsPipelineState(RHICmdList, State.GraphicsPSOInit, State.StencilRef);
+		VertexBindings = FSlateDrawShaderBindings::Create(GraphBuilder, VertexShader);
+		VertexShader->SetMaterialShaderParameters(*VertexBindings, Scene, ViewUniformBuffer, MaterialRenderProxy, EffectiveMaterial);
 
-		{
-			FRHIBatchedShaderParameters& BatchedParameters = RHICmdList.GetScratchShaderParameters();
+		const bool bDrawDisabled = EnumHasAllFlags(RenderBatch->DrawEffects, ESlateDrawEffect::DisabledEffect);
 
-			VertexShader->SetViewProjection(BatchedParameters, Inputs.ElementsMatrix);
-			VertexShader->SetMaterialShaderParameters(BatchedParameters, Scene, ViewUniformBuffer, MaterialRenderProxy, EffectiveMaterial);
+		PixelBindings = FSlateDrawShaderBindings::Create(GraphBuilder, PixelShader);
+		PixelShader->SetMaterialShaderParameters(*PixelBindings, Scene, ViewUniformBuffer, MaterialRenderProxy, EffectiveMaterial, ShaderParams);
+		PixelShader->SetDisplayGammaAndContrast(*PixelBindings, FinalGamma, FinalContrast);
+		PixelShader->SetDrawFlags(*PixelBindings, bDrawDisabled);
 
-			RHICmdList.SetBatchedShaderParameters(VertexShader.GetVertexShader(), BatchedParameters);
-		}
+		auto* MaskResource = static_cast<TSlateTexture<FTextureRHIRef>*>(MaterialShaderResource->GetTextureMaskResource());
 
+		if (MaskResource)
 		{
-			FRHIBatchedShaderParameters& BatchedParameters = RHICmdList.GetScratchShaderParameters();
-
-			PixelShader->SetParameters(BatchedParameters, Scene, ViewUniformBuffer, MaterialRenderProxy, EffectiveMaterial, ShaderParams);
-			const float FinalGamma = EnumHasAnyFlags(DrawFlags, ESlateBatchDrawFlag::ReverseGamma) ? 1.0f / Inputs.EngineGamma : EnumHasAnyFlags(DrawFlags, ESlateBatchDrawFlag::NoGamma) ? 1.0f : Inputs.DisplayGamma;
-			const float FinalContrast = EnumHasAnyFlags(DrawFlags, ESlateBatchDrawFlag::NoGamma) ? 1 : Inputs.DisplayContrast;
-			PixelShader->SetDisplayGammaAndContrast(BatchedParameters, FinalGamma, FinalContrast);
-			const bool bDrawDisabled = EnumHasAllFlags(DrawEffects, ESlateDrawEffect::DisabledEffect);
-			PixelShader->SetDrawFlags(BatchedParameters, bDrawDisabled);
-
-			if (MaskResource)
-			{
-				FTextureRHIRef TextureRHI = ((TSlateTexture<FTextureRHIRef>*)MaskResource)->GetTypedResource();
-
-				PixelShader->SetAdditionalTexture(BatchedParameters, TextureRHI, TStaticSamplerState<SF_Bilinear, AM_Clamp, AM_Clamp, AM_Clamp>::GetRHI());
-			}
-
-			RHICmdList.SetBatchedShaderParameters(PixelShader.GetPixelShader(), BatchedParameters);
+			PixelShader->SetAdditionalTexture(*PixelBindings, MaskResource->GetTypedResource(), TStaticSamplerState<SF_Bilinear, AM_Clamp, AM_Clamp, AM_Clamp>::GetRHI());
 		}
 
-		if (bUseInstancing)
-		{
-			RenderBatch.InstanceData->BindStreamSource(RHICmdList, 1, RenderBatch.InstanceOffset);
-
-			RHICmdList.SetStreamSource(0, ElementsVertexBuffer, RenderBatch.VertexOffset * sizeof(FSlateVertex));
-			RHICmdList.DrawIndexedPrimitive(ElementsIndexBuffer, 0, 0, RenderBatch.NumVertices, RenderBatch.IndexOffset, PrimitiveCount, RenderBatch.InstanceCount);
-		}
-		else
-		{
-			RHICmdList.SetStreamSource(0, ElementsVertexBuffer, RenderBatch.VertexOffset * sizeof(FSlateVertex));
-			RHICmdList.SetStreamSource(1, nullptr, 0);
-			RHICmdList.DrawIndexedPrimitive(ElementsIndexBuffer, 0, 0, RenderBatch.NumVertices, RenderBatch.IndexOffset, PrimitiveCount, 1);
-		}
+		BlendState = GetMaterialBlendState(MaskResource, EffectiveMaterial);
 	}
 	else
 	{
@@ -1104,7 +1115,7 @@ void DrawSlateRenderBatch(
 		TShaderRef<FSlateDebugBatchingPS> BatchingPixelShader;
 		if (CVarShowSlateBatching.GetValueOnRenderThread() != 0)
 		{
-			BatchingPixelShader = TShaderMapRef<FSlateDebugBatchingPS>(ShaderMap);
+			BatchingPixelShader = TShaderMapRef<FSlateDebugBatchingPS>(Inputs.ShaderMap);
 			PixelShader = BatchingPixelShader;
 		}
 		else
@@ -1130,22 +1141,22 @@ void DrawSlateRenderBatch(
 				}
 			}
 
-			PixelShader = GetTexturePixelShader(ShaderMap, ShaderType, DrawEffects, bUseTextureGrayscale, bIsVirtualTexture);
+			PixelShader = GetTexturePixelShader(Inputs.ShaderMap, ShaderType, DrawEffects, bUseTextureGrayscale, bIsVirtualTexture);
 		}
 
 #if WITH_SLATE_VISUALIZERS
 		if (BatchingPixelShader.IsValid())
 		{
-			State.GraphicsPSOInit.BlendState = TStaticBlendState<CW_RGBA, BO_Add, BF_SourceAlpha, BF_InverseSourceAlpha, BO_Add, BF_One, BF_InverseSourceAlpha>::GetRHI();
+			BlendState = TStaticBlendState<CW_RGBA, BO_Add, BF_SourceAlpha, BF_InverseSourceAlpha, BO_Add, BF_One, BF_InverseSourceAlpha>::GetRHI();
 		}
 		else if (CVarShowSlateOverdraw.GetValueOnRenderThread() != 0)
 		{
-			State.GraphicsPSOInit.BlendState = TStaticBlendState<CW_RGB, BO_Add, BF_One, BF_One, BO_Add, BF_Zero, BF_InverseSourceAlpha>::GetRHI();
+			BlendState = TStaticBlendState<CW_RGB, BO_Add, BF_One, BF_One, BO_Add, BF_Zero, BF_InverseSourceAlpha>::GetRHI();
 		}
 		else
 #endif
 		{
-			State.GraphicsPSOInit.BlendState =
+			BlendState =
 				EnumHasAllFlags(DrawFlags, ESlateBatchDrawFlag::NoBlending)
 				? TStaticBlendState<>::GetRHI()
 				: (EnumHasAllFlags(DrawFlags, ESlateBatchDrawFlag::PreMultipliedAlpha)
@@ -1154,29 +1165,6 @@ void DrawSlateRenderBatch(
 				;
 		}
 
-		if (EnumHasAllFlags(DrawFlags, ESlateBatchDrawFlag::Wireframe) || Inputs.bWireframe)
-		{
-			State.GraphicsPSOInit.RasterizerState = TStaticRasterizerState<FM_Wireframe>::GetRHI();
-
-			if (Inputs.bWireframe)
-			{
-				State.GraphicsPSOInit.BlendState = TStaticBlendState<>::GetRHI();
-			}
-		}
-		else
-		{
-			State.GraphicsPSOInit.RasterizerState = TStaticRasterizerState<FM_Solid>::GetRHI();
-		}
-
-		TShaderMapRef<FSlateElementVS> GlobalVertexShader(ShaderMap);
-
-		State.GraphicsPSOInit.BoundShaderState.VertexDeclarationRHI = GSlateVertexDeclaration.VertexDeclarationRHI;
-		State.GraphicsPSOInit.BoundShaderState.VertexShaderRHI = GlobalVertexShader.GetVertexShader();
-		State.GraphicsPSOInit.BoundShaderState.PixelShaderRHI = PixelShader.GetPixelShader();
-		State.GraphicsPSOInit.PrimitiveType = GetRHIPrimitiveType(RenderBatch.DrawPrimitiveType);
-
-		SetGraphicsPipelineState(RHICmdList, State.GraphicsPSOInit, State.StencilRef);
-
 		FRHISamplerState* SamplerState = TStaticSamplerState<SF_Bilinear, AM_Clamp, AM_Clamp, AM_Clamp>::GetRHI();
 		FRHITexture* TextureRHI = GWhiteTexture->TextureRHI;
 		bool bIsVirtualTexture = false;
@@ -1221,41 +1209,135 @@ void DrawSlateRenderBatch(
 			SamplerState = GetSamplerState(DrawFlags, Filter);
 		}
 
-		{
-			FRHIBatchedShaderParameters& BatchedParameters = RHICmdList.GetScratchShaderParameters();
+		PixelBindings = FSlateDrawShaderBindings::Create(GraphBuilder, PixelShader);
 
 #if WITH_SLATE_VISUALIZERS
-			if (BatchingPixelShader.IsValid())
-			{
-				BatchingPixelShader->SetBatchColor(BatchedParameters, Inputs.BatchColor);
-			}
+		if (BatchingPixelShader.IsValid())
+		{
+			BatchingPixelShader->SetBatchColor(*PixelBindings, Inputs.BatchColor);
+		}
 #endif
 
-			if (bIsVirtualTexture && TextureResource != nullptr)
-			{
-				PixelShader->SetVirtualTextureParameters(BatchedParameters, static_cast<FVirtualTexture2DResource*>(TextureResource));
-			}
-			else
-			{
-				PixelShader->SetTexture(BatchedParameters, TextureRHI, SamplerState);
-			}
+		if (bIsVirtualTexture && TextureResource != nullptr)
+		{
+			PixelShader->SetVirtualTextureParameters(*PixelBindings, static_cast<FVirtualTexture2DResource*>(TextureResource));
+		}
+		else
+		{
+			PixelShader->SetTexture(*PixelBindings, TextureRHI, SamplerState);
+		}
 
-			PixelShader->SetShaderParams(BatchedParameters, ShaderParams);
-			const float FinalGamma = EnumHasAnyFlags(DrawFlags, ESlateBatchDrawFlag::ReverseGamma) ? (1.0f / Inputs.EngineGamma) : EnumHasAnyFlags(DrawFlags, ESlateBatchDrawFlag::NoGamma) ? 1.0f : Inputs.DisplayGamma;
-			const float FinalContrast = EnumHasAnyFlags(DrawFlags, ESlateBatchDrawFlag::NoGamma) ? 1 : Inputs.DisplayContrast;
-			PixelShader->SetDisplayGammaAndInvertAlphaAndContrast(BatchedParameters, FinalGamma, EnumHasAllFlags(DrawEffects, ESlateDrawEffect::InvertAlpha) ? 1.0f : 0.0f, FinalContrast);
+		PixelShader->SetShaderParams(*PixelBindings, ShaderParams);
+		PixelShader->SetDisplayGammaAndInvertAlphaAndContrast(*PixelBindings, FinalGamma, EnumHasAllFlags(DrawEffects, ESlateDrawEffect::InvertAlpha) ? 1.0f : 0.0f, FinalContrast);
+	}
 
-			RHICmdList.SetBatchedShaderParameters(PixelShader.GetPixelShader(), BatchedParameters);
-		}
+	FSlateRenderBatchOp* RenderBatchOp = GraphBuilder.AllocPOD<FSlateRenderBatchOp>();
+	RenderBatchOp->RenderBatch        = RenderBatch;
+	RenderBatchOp->ClippingStateOp    = ClippingStateOp;
+	RenderBatchOp->ShaderResourceType = ResourceType;
+	RenderBatchOp->VertexBindings     = VertexBindings;
+	RenderBatchOp->PixelBindings      = PixelBindings;
+	RenderBatchOp->BlendState         = BlendState;
+	RenderBatchOp->Next               = nullptr;
+	return RenderBatchOp;
+}
+
+struct FSlateRenderBatchDrawInputs
+{
+	FGlobalShaderMap* ShaderMap;
+	FSlateElementsBuffers ElementsBuffers;
+	FIntRect ElementsViewRect;
+	bool bWireframe;
+};
+
+void DrawSlateRenderBatch(
+	FRHICommandList& RHICmdList,
+	FSlateRenderBatchDrawState& State,
+	const FSlateRenderBatchDrawInputs& Inputs,
+	const FSlateRenderBatchOp& RenderBatchOp)
+{
+	const FSlateClippingOp* ClippingStateOp = RenderBatchOp.ClippingStateOp;
+	const FSlateRenderBatch& RenderBatch    = *RenderBatchOp.RenderBatch;
+
+	if (State.LastClippingOp != ClippingStateOp)
+	{
+		GetSlateClippingPipelineState(ClippingStateOp, State.GraphicsPSOInit.DepthStencilState, State.StencilRef);
+		SetSlateClipping(RHICmdList, ClippingStateOp, Inputs.ElementsViewRect);
+		State.LastClippingOp = ClippingStateOp;
+	}
+
+	FRHIBuffer* ElementsVertexBuffer = Inputs.ElementsBuffers.VertexBuffer->GetRHI();
+	FRHIBuffer* ElementsIndexBuffer  = Inputs.ElementsBuffers.IndexBuffer->GetRHI();
+
+	State.GraphicsPSOInit.BlendState = RenderBatchOp.BlendState;
+
+	if (EnumHasAllFlags(RenderBatch.DrawFlags, ESlateBatchDrawFlag::Wireframe))
+	{
+		State.GraphicsPSOInit.RasterizerState = TStaticRasterizerState<FM_Wireframe>::GetRHI();
+	}
+	else
+	{
+		State.GraphicsPSOInit.RasterizerState = TStaticRasterizerState<FM_Solid>::GetRHI();
+	}
+
+	check(RenderBatch.NumIndices > 0);
+	const uint32 PrimitiveCount = RenderBatch.DrawPrimitiveType == ESlateDrawPrimitive::LineList ? RenderBatch.NumIndices / 2 : RenderBatch.NumIndices / 3;
+
+	if (RenderBatchOp.ShaderResourceType == ESlateShaderResource::Material)
+	{
+		const bool bUseInstancing = RenderBatch.InstanceCount > 0 && RenderBatch.InstanceData != nullptr;
+
+		State.GraphicsPSOInit.BoundShaderState.VertexDeclarationRHI = bUseInstancing ? GSlateInstancedVertexDeclaration.VertexDeclarationRHI : GSlateVertexDeclaration.VertexDeclarationRHI;
+		State.GraphicsPSOInit.BoundShaderState.VertexShaderRHI = RenderBatchOp.VertexBindings->Shader.GetVertexShader();
+		State.GraphicsPSOInit.BoundShaderState.PixelShaderRHI  = RenderBatchOp.PixelBindings->Shader.GetPixelShader();
+		State.GraphicsPSOInit.PrimitiveType = GetRHIPrimitiveType(RenderBatch.DrawPrimitiveType);
+
+		SetGraphicsPipelineState(RHICmdList, State.GraphicsPSOInit, State.StencilRef);
+
+		RenderBatchOp.VertexBindings->SetOnCommandList(RHICmdList);
+		RenderBatchOp.PixelBindings->SetOnCommandList(RHICmdList);
 
+		if (bUseInstancing)
 		{
-			FRHIBatchedShaderParameters& BatchedParameters = RHICmdList.GetScratchShaderParameters();
+			RenderBatch.InstanceData->BindStreamSource(RHICmdList, 1, RenderBatch.InstanceOffset);
 
-			GlobalVertexShader->SetViewProjection(BatchedParameters, Inputs.ElementsMatrix);
+			RHICmdList.SetStreamSource(0, ElementsVertexBuffer, RenderBatch.VertexOffset * sizeof(FSlateVertex));
+			RHICmdList.DrawIndexedPrimitive(ElementsIndexBuffer, 0, 0, RenderBatch.NumVertices, RenderBatch.IndexOffset, PrimitiveCount, RenderBatch.InstanceCount);
+		}
+		else
+		{
+			RHICmdList.SetStreamSource(0, ElementsVertexBuffer, RenderBatch.VertexOffset * sizeof(FSlateVertex));
+			RHICmdList.SetStreamSource(1, nullptr, 0);
+			RHICmdList.DrawIndexedPrimitive(ElementsIndexBuffer, 0, 0, RenderBatch.NumVertices, RenderBatch.IndexOffset, PrimitiveCount, 1);
+		}
+	}
+	else
+	{
+		if (EnumHasAllFlags(RenderBatch.DrawFlags, ESlateBatchDrawFlag::Wireframe) || Inputs.bWireframe)
+		{
+			State.GraphicsPSOInit.RasterizerState = TStaticRasterizerState<FM_Wireframe>::GetRHI();
 
-			RHICmdList.SetBatchedShaderParameters(GlobalVertexShader.GetVertexShader(), BatchedParameters);
+			if (Inputs.bWireframe)
+			{
+				State.GraphicsPSOInit.BlendState = TStaticBlendState<>::GetRHI();
+			}
+		}
+		else
+		{
+			State.GraphicsPSOInit.RasterizerState = TStaticRasterizerState<FM_Solid>::GetRHI();
 		}
 
+		TShaderMapRef<FSlateElementVS> GlobalVertexShader(Inputs.ShaderMap);
+
+		State.GraphicsPSOInit.BoundShaderState.VertexDeclarationRHI = GSlateVertexDeclaration.VertexDeclarationRHI;
+		State.GraphicsPSOInit.BoundShaderState.VertexShaderRHI = GlobalVertexShader.GetVertexShader();
+		State.GraphicsPSOInit.BoundShaderState.PixelShaderRHI = RenderBatchOp.PixelBindings->Shader.GetPixelShader();
+		State.GraphicsPSOInit.PrimitiveType = GetRHIPrimitiveType(RenderBatch.DrawPrimitiveType);
+
+		SetGraphicsPipelineState(RHICmdList, State.GraphicsPSOInit, State.StencilRef);
+
+		RenderBatchOp.PixelBindings->SetOnCommandList(RHICmdList);
+
 		RHICmdList.SetStreamSource(0, ElementsVertexBuffer, RenderBatch.VertexOffset * sizeof(FSlateVertex));
 		RHICmdList.DrawIndexedPrimitive(ElementsIndexBuffer, 0, 0, RenderBatch.NumVertices, RenderBatch.IndexOffset, PrimitiveCount, RenderBatch.InstanceCount);
 	}
@@ -1265,6 +1347,7 @@ BEGIN_SHADER_PARAMETER_STRUCT(FSlateRenderBatchParameters, )
 	SHADER_PARAMETER_STRUCT_INCLUDE(FSceneTextureExtractsParameters, SceneTextures)
 	RDG_BUFFER_ACCESS(ElementsVertexBuffer, ERHIAccess::VertexOrIndexBuffer)
 	RDG_BUFFER_ACCESS(ElementsIndexBuffer, ERHIAccess::VertexOrIndexBuffer)
+	SHADER_PARAMETER_STRUCT_REF(FSlateViewUniformParameters, SlateView)
 	RENDER_TARGET_BINDING_SLOTS()
 END_SHADER_PARAMETER_STRUCT()
 
@@ -1283,34 +1366,50 @@ void AddSlateDrawElementsPass(
 	const float EngineGamma  = GEngine ? GEngine->GetDisplayGamma() : 2.2f;
 	const float DisplayGamma = Inputs.bAllowGammaCorrection && !Inputs.bElementsTextureIsHDRDisplay ? EngineGamma : 1.0f;
 
+	FGlobalShaderMap* ShaderMap = GetGlobalShaderMap(GMaxRHIFeatureLevel);
+
 	FSlateRHIResourceManager& ResourceManager = RenderingPolicy.GetResourceManagerRHI();
 
-	FSlateSceneViewAllocator* SceneViewAllocator = FSlateSceneViewAllocator::Create(GraphBuilder, ResourceManager);
+	const FSlateSceneViewAllocateInputs SceneViewAllocateInputs =
+	{
+		  .TextureExtent        = ElementsTextureExtent
+		, .ViewRect             = Inputs.SceneViewRect
+		, .ViewProjectionMatrix = Inputs.ElementsMatrix
+		, .CursorPosition       = Inputs.CursorPosition
+		, .Time                 = Inputs.Time
+		, .ViewportScaleUI      = Inputs.ViewportScaleUI
+	};
+
+	FSlateSceneViewAllocator* SceneViewAllocator = FSlateSceneViewAllocator::Create(GraphBuilder, ResourceManager, SceneViewAllocateInputs);
 
 #if WITH_SLATE_VISUALIZERS
 	FRandomStream BatchColors(1337);
 #endif
 
-	FSlateRenderBatchDrawInputs* RenderBatchDrawInputs = GraphBuilder.AllocPOD<FSlateRenderBatchDrawInputs>();
-
-	*RenderBatchDrawInputs =
+	const FSlateRenderBatchCreateInputs RenderBatchCreateInputs
 	{
-		  .ElementsMatrix        = Inputs.ElementsMatrix
-		, .ElementsOffset        = Inputs.ElementsOffset
-		, .ElementsViewRect      = ElementsTexture.ViewRect
-		, .SceneViewRect         = Inputs.SceneViewRect
+		  .ShaderMap             = ShaderMap
 		, .SceneViewAllocator    = SceneViewAllocator
 		, .TextureLODGroups      = GetTextureLODGroups()
-		, .ElementsBuffers       = Inputs.ElementsBuffers
 		, .DisplayGamma          = DisplayGamma
 		, .DisplayContrast       = GSlateContrast
 		, .EngineGamma           = EngineGamma
-		, .bWireframe            = Inputs.bWireframe
 #if WITH_SLATE_VISUALIZERS
 		, .BatchColor            = FLinearColor(BatchColors.GetUnitVector())
 #endif
 	};
 
+	// Draw inputs are passed into RDG lambdas and need to be allocated by RDG.
+	FSlateRenderBatchDrawInputs* RenderBatchDrawInputs = GraphBuilder.AllocPOD<FSlateRenderBatchDrawInputs>();
+
+	*RenderBatchDrawInputs =
+	{
+		  .ShaderMap             = ShaderMap
+		, .ElementsBuffers       = Inputs.ElementsBuffers
+		, .ElementsViewRect      = ElementsTexture.ViewRect
+		, .bWireframe            = Inputs.bWireframe
+	};
+
 	ERenderTargetLoadAction ElementsLoadAction = Inputs.ElementsLoadAction;
 
 	const auto ConsumeLoadAction = [] (ERenderTargetLoadAction& InOutLoadAction)
@@ -1320,12 +1419,20 @@ void AddSlateDrawElementsPass(
 		return LoadAction;
 	};
 
+	TUniformBufferRef<FSlateViewUniformParameters> SlateViewUniformBuffer;
+
 	FSlateRenderBatchParameters* NoneStencilActionPassParameters = GraphBuilder.AllocParameters<FSlateRenderBatchParameters>();
 	NoneStencilActionPassParameters->SceneTextures = GetSceneTextureExtracts().GetShaderParameters();
 	NoneStencilActionPassParameters->ElementsVertexBuffer = Inputs.ElementsBuffers.VertexBuffer;
 	NoneStencilActionPassParameters->ElementsIndexBuffer  = Inputs.ElementsBuffers.IndexBuffer;
 	NoneStencilActionPassParameters->RenderTargets[0] = FRenderTargetBinding(Inputs.ElementsTexture, ERenderTargetLoadAction::ELoad);
 
+	{
+		FSlateViewUniformParameters UniformParameters;
+		UniformParameters.ViewProjection = Inputs.ElementsMatrix;
+		NoneStencilActionPassParameters->SlateView = TUniformBufferRef<FSlateViewUniformParameters>::CreateUniformBufferImmediate(UniformParameters, UniformBuffer_SingleFrame);
+	}
+
 	FSlateRenderBatchParameters* ClearStencilActionPassParameters = nullptr;
 	FSlateRenderBatchParameters* WriteStencilActionPassParameters = nullptr;
 
@@ -1344,12 +1451,13 @@ void AddSlateDrawElementsPass(
 
 	FSlateClippingCreateContext ClippingCreateContext;
 
-	TArray<FSlateRenderBatchOp, FRDGArrayAllocator> RenderBatchOps;
-	RenderBatchOps.Reserve(RenderBatches.Num());
+	FSlateRenderBatchOp* RenderBatchHeadOp = nullptr;
+	FSlateRenderBatchOp* RenderBatchTailOp = nullptr;
+	int32 NumRenderBatchOps = 0;
 
 	const auto FlushDrawElementsPass = [&]
 	{
-		if (RenderBatchOps.IsEmpty())
+		if (!NumRenderBatchOps)
 		{
 			return;
 		}
@@ -1361,9 +1469,7 @@ void AddSlateDrawElementsPass(
 			LastPassParameters->RenderTargets[0].SetLoadAction(LoadAction);
 		}
 
-		const int32 NumRenderBatchOps = RenderBatchOps.Num();
-
-		FRDGPass* Pass = GraphBuilder.AddPass(RDG_EVENT_NAME("ElementBatch"), LastPassParameters, ERDGPassFlags::Raster, [Inputs = RenderBatchDrawInputs, RenderBatchOps = MoveTemp(RenderBatchOps)] (FRHICommandList& RHICmdList)
+		FRDGPass* Pass = GraphBuilder.AddPass(RDG_EVENT_NAME("ElementBatch"), LastPassParameters, ERDGPassFlags::Raster, [Inputs = RenderBatchDrawInputs, RenderBatchHeadOp] (FRHICommandList& RHICmdList)
 		{
 			RHICmdList.SetViewport(Inputs->ElementsViewRect.Min.X, Inputs->ElementsViewRect.Min.Y, 0.0f, Inputs->ElementsViewRect.Max.X, Inputs->ElementsViewRect.Max.Y, 1.0f);
 			RHICmdList.SetScissorRect(false, 0, 0, 0, 0);
@@ -1372,14 +1478,19 @@ void AddSlateDrawElementsPass(
 			RHICmdList.ApplyCachedRenderTargets(DrawState.GraphicsPSOInit);
 			DrawState.GraphicsPSOInit.DepthStencilState = TStaticDepthStencilState<false, CF_Always>::GetRHI();
 
-			for (const FSlateRenderBatchOp& RenderBatchOp : RenderBatchOps)
+			const FSlateRenderBatchOp* RenderBatchOp = RenderBatchHeadOp;
+			const FSlateRenderBatchOp* LastRenderBatchOp = nullptr;
+
+			for (; RenderBatchOp != nullptr; RenderBatchOp = RenderBatchOp->Next)
 			{
-				DrawSlateRenderBatch(RHICmdList, DrawState, *Inputs, RenderBatchOp);
+				DrawSlateRenderBatch(RHICmdList, DrawState, *Inputs, *RenderBatchOp);
+				LastRenderBatchOp = RenderBatchOp;
 			}
 		});
 
 		GraphBuilder.SetPassWorkload(Pass, NumRenderBatchOps);
-		check(RenderBatchOps.IsEmpty());
+		RenderBatchHeadOp = RenderBatchTailOp = nullptr;
+		NumRenderBatchOps = 0;
 	};
 
 	int32 NextRenderBatchIndex = FirstBatchIndex;
@@ -1491,37 +1602,27 @@ void AddSlateDrawElementsPass(
 		}
 		case ESlateRenderBatchType::Primitive:
 		{
-			const FSlateSceneView* SceneView = nullptr;
-
-			if (NextRenderBatch.ShaderResource && NextRenderBatch.ShaderResource->GetType() == ESlateShaderResource::Material)
+			if (FSlateRenderBatchOp* RenderBatchOp = CreateSlateRenderBatchOp(GraphBuilder, RenderBatchCreateInputs, &NextRenderBatch, NextClippingOp))
 			{
-				// Skip material render batches when the engine is not available.
-				if (!GEngine)
+				if (!RenderBatchTailOp)
 				{
-					break;
+					RenderBatchHeadOp = RenderBatchTailOp = RenderBatchOp;
 				}
-
-				const FSlateSceneViewAllocateInputs AllocateInputs =
+				else
 				{
-					  .TextureExtent        = ElementsTextureExtent
-					, .ViewRect             = Inputs.SceneViewRect
-					, .ViewProjectionMatrix = Inputs.ElementsMatrix
-					, .CursorPosition       = Inputs.CursorPosition
-					, .Time                 = Inputs.Time
-					, .ViewportScaleUI      = Inputs.ViewportScaleUI
-				};
-
-				SceneView = SceneViewAllocator->BeginAllocateSceneView(GraphBuilder, ResourceManager, NextRenderBatch.SceneIndex, AllocateInputs);
+					RenderBatchTailOp->Next = RenderBatchOp;
+					RenderBatchTailOp = RenderBatchOp;
+				}
+				NumRenderBatchOps++;
 			}
 
-			RenderBatchOps.Emplace(&NextRenderBatch, NextClippingOp, SceneView);
 			break;
 		}
 		default: checkNoEntry();
 		}
 	}
 
-	if (!RenderBatchOps.IsEmpty())
+	if (NumRenderBatchOps > 0)
 	{
 		FlushDrawElementsPass();
 	}
diff --git a/Engine/Source/Runtime/SlateRHIRenderer/Private/SlateShaders.cpp b/Engine/Source/Runtime/SlateRHIRenderer/Private/SlateShaders.cpp
--- a/Engine/Source/Runtime/SlateRHIRenderer/Private/SlateShaders.cpp
+++ b/Engine/Source/Runtime/SlateRHIRenderer/Private/SlateShaders.cpp
@@ -138,26 +138,4 @@ void FSlateMaskingVertexDeclaration::InitRHI(FRHICommandListBase& RHICmdList)
 void FSlateMaskingVertexDeclaration::ReleaseRHI()
 {
 	VertexDeclarationRHI.SafeRelease();
-}
-
-
-/************************************************************************/
-/* FSlateDefaultVertexShader                                            */
-/************************************************************************/
-
-FSlateElementVS::FSlateElementVS( const ShaderMetaType::CompiledShaderInitializerType& Initializer )
-	: FGlobalShader(Initializer)
-{
-	ViewProjection.Bind(Initializer.ParameterMap, TEXT("ViewProjection"));
-	VertexShaderParams.Bind( Initializer.ParameterMap, TEXT("VertexShaderParams"));
-}
-
-void FSlateElementVS::SetViewProjection(FRHIBatchedShaderParameters& BatchedParameters, const FMatrix44f& InViewProjection )
-{
-	SetShaderValue(BatchedParameters, ViewProjection, InViewProjection );
-}
-
-void FSlateElementVS::SetShaderParameters(FRHIBatchedShaderParameters& BatchedParameters, const FVector4f& ShaderParams )
-{
-	SetShaderValue(BatchedParameters, VertexShaderParams, ShaderParams );
 }
\ No newline at end of file
diff --git a/Engine/Source/Runtime/SlateRHIRenderer/Private/SlateShaders.h b/Engine/Source/Runtime/SlateRHIRenderer/Private/SlateShaders.h
--- a/Engine/Source/Runtime/SlateRHIRenderer/Private/SlateShaders.h
+++ b/Engine/Source/Runtime/SlateRHIRenderer/Private/SlateShaders.h
@@ -14,6 +14,7 @@
 #include "TextureResource.h"
 #include "RenderUtils.h"
 #include "ShaderParameterStruct.h"
+#include "MeshDrawShaderBindings.h"
 
 extern EColorVisionDeficiency GSlateColorDeficiencyType;
 extern int32 GSlateColorDeficiencySeverity;
@@ -63,68 +64,33 @@ public:
 	virtual void ReleaseRHI() override;
 };
 
-/** The slate Vertex shader representation */
 class FSlateElementVS : public FGlobalShader
 {
-	DECLARE_SHADER_TYPE(FSlateElementVS, Global);
 public:
-	/** Indicates that this shader should be cached */
-	static bool ShouldCompilePermutation(const FGlobalShaderPermutationParameters& Parameters) { return true; }
-
-	/** Constructor.  Binds all parameters used by the shader */
-	FSlateElementVS( const ShaderMetaType::CompiledShaderInitializerType& Initializer );
+	DECLARE_GLOBAL_SHADER(FSlateElementVS);
 
-	FSlateElementVS() {}
+	FSlateElementVS() = default;
 
-	/** 
-	 * Sets the view projection parameter
-	 *
-	 * @param InViewProjection	The ViewProjection matrix to use when this shader is bound 
-	 */
-	void SetViewProjection(FRHIBatchedShaderParameters& BatchedParameters, const FMatrix44f& InViewProjection );
-
-	/** 
-	 * Sets shader parameters for use in this shader
-	 *
-	 * @param ShaderParams	The shader params to be used
-	 */
-	void SetShaderParameters(FRHIBatchedShaderParameters& BatchedParameters, const FVector4f& ShaderParams );
-
-	/** Serializes the shader data */
-	//virtual bool Serialize( FArchive& Ar ) override;
-
-private:
-	/** ViewProjection parameter used by the shader */
-	LAYOUT_FIELD(FShaderParameter, ViewProjection)
-	/** Shader parmeters used by the shader */
-	LAYOUT_FIELD(FShaderParameter, VertexShaderParams)
+	FSlateElementVS(const ShaderMetaType::CompiledShaderInitializerType& Initializer)
+		: FGlobalShader(Initializer)
+	{}
 };
 
-/** 
- * Base class slate pixel shader for all elements
- */
 class FSlateElementPS : public FGlobalShader
 {
-	DECLARE_TYPE_LAYOUT(FSlateElementPS, NonVirtual);
 public:
-	/** Indicates that this shader should be cached */
-	static bool ShouldCompilePermutation(const FGlobalShaderPermutationParameters& Parameters) 
-	{ 
-		return true; 
-	}
+	DECLARE_GLOBAL_SHADER(FSlateElementPS);
 
-	FSlateElementPS()
-	{
-	}
+	FSlateElementPS() = default;
 
-	/** Constructor.  Binds all parameters used by the shader */
-	FSlateElementPS( const ShaderMetaType::CompiledShaderInitializerType& Initializer )
-		: FGlobalShader( Initializer )
+	FSlateElementPS(const ShaderMetaType::CompiledShaderInitializerType& Initializer)
+		: FGlobalShader(Initializer)
 	{
 		TextureParameter.Bind(Initializer.ParameterMap, TEXT("ElementTexture"));
 		TextureParameterSampler.Bind(Initializer.ParameterMap, TEXT("ElementTextureSampler"));
 		InPageTableTexture.Bind(Initializer.ParameterMap, TEXT("InPageTableTexture"));
-		VTPackedPageTableUniform.Bind(Initializer.ParameterMap, TEXT("VTPackedPageTableUniform"));
+		VTPackedPageTableUniform0.Bind(Initializer.ParameterMap, TEXT("VTPackedPageTableUniform0"));
+		VTPackedPageTableUniform1.Bind(Initializer.ParameterMap, TEXT("VTPackedPageTableUniform1"));
 		VTPackedUniform.Bind(Initializer.ParameterMap, TEXT("VTPackedUniform"));
 		ShaderParams.Bind(Initializer.ParameterMap, TEXT("ShaderParams"));
 		ShaderParams2.Bind(Initializer.ParameterMap, TEXT("ShaderParams2"));
@@ -140,9 +106,9 @@ public:
 	 * @param Texture	Texture resource to use when this pixel shader is bound
 	 * @param SamplerState	Sampler state to use when sampling this texture
 	 */
-	void SetTexture(FRHIBatchedShaderParameters& BatchedParameters, FRHITexture* InTexture, const FSamplerStateRHIRef SamplerState )
+	void SetTexture(FMeshDrawSingleShaderBindings& ShaderBindings, FRHITexture* InTexture, const FSamplerStateRHIRef SamplerState)
 	{
-		SetTextureParameter(BatchedParameters, TextureParameter, TextureParameterSampler, SamplerState, InTexture );
+		ShaderBindings.AddTexture(TextureParameter, TextureParameterSampler, SamplerState, InTexture);
 	}
 
 	/**
@@ -150,7 +116,7 @@ public:
 	 *
 	 * @param InVirtualTexture	Virtual Texture resource to use when this pixel shader is bound
 	 */
-	void SetVirtualTextureParameters(FRHIBatchedShaderParameters& BatchedParameters, FVirtualTexture2DResource* InVirtualTexture)
+	void SetVirtualTextureParameters(FMeshDrawSingleShaderBindings& ShaderBindings, FVirtualTexture2DResource* InVirtualTexture)
 	{
 		if (InVirtualTexture == nullptr)
 		{
@@ -161,11 +127,7 @@ public:
 		uint32 LayerIndex = 0;
 
 		FRHIShaderResourceView* PhysicalView = AllocatedVT->GetPhysicalTextureSRV(LayerIndex, InVirtualTexture->bSRGB);
-		
-		SetSRVParameter(BatchedParameters, TextureParameter, PhysicalView);
-		SetSamplerParameter(BatchedParameters, TextureParameterSampler, InVirtualTexture->SamplerStateRHI);
-		SetTextureParameter(BatchedParameters, InPageTableTexture, AllocatedVT->GetPageTableTexture(0u));
-		
+
 		FUintVector4 PageTableUniform[2];
 		FUintVector4 Uniform;
 		// VTParams.X = MipLevel, VTParams.Y = LayerIndex
@@ -174,9 +136,13 @@ public:
 		AllocatedVT->GetPackedPageTableUniform(PageTableUniform);
 		AllocatedVT->GetPackedUniform(&Uniform, LayerIndex);
 
-		SetShaderValueArray(BatchedParameters, VTPackedPageTableUniform, PageTableUniform, UE_ARRAY_COUNT(PageTableUniform));
-		SetShaderValue(BatchedParameters, VTPackedUniform, Uniform);
-		SetShaderValue(BatchedParameters, VTShaderParams, VTParams);
+		ShaderBindings.Add(TextureParameter, PhysicalView);
+		ShaderBindings.Add(TextureParameterSampler, InVirtualTexture->SamplerStateRHI);
+		ShaderBindings.Add(InPageTableTexture, AllocatedVT->GetPageTableTexture(0u));
+		ShaderBindings.Add(VTPackedPageTableUniform0, PageTableUniform[0]);
+		ShaderBindings.Add(VTPackedPageTableUniform1, PageTableUniform[1]);
+		ShaderBindings.Add(VTPackedUniform, Uniform);
+		ShaderBindings.Add(VTShaderParams, VTParams);
 	}
 
 	/**
@@ -184,10 +150,10 @@ public:
 	 * 
 	 * @param InShaderParams Shader params to use
 	 */
-	void SetShaderParams(FRHIBatchedShaderParameters& BatchedParameters, const FShaderParams& InShaderParams)
+	void SetShaderParams(FMeshDrawSingleShaderBindings& ShaderBindings, const FShaderParams& InShaderParams)
 	{
-		SetShaderValue(BatchedParameters, ShaderParams, InShaderParams.PixelParams);
-		SetShaderValue(BatchedParameters, ShaderParams2, InShaderParams.PixelParams2);
+		ShaderBindings.Add(ShaderParams, InShaderParams.PixelParams);
+		ShaderBindings.Add(ShaderParams2, InShaderParams.PixelParams2);
 	}
 
 	/**
@@ -195,20 +161,19 @@ public:
 	 *
 	 * @param DisplayGamma The display gamma to use
 	 */
-	void SetDisplayGammaAndInvertAlphaAndContrast(FRHIBatchedShaderParameters& BatchedParameters, float InDisplayGamma, float bInvertAlpha, float InContrast)
+	void SetDisplayGammaAndInvertAlphaAndContrast(FMeshDrawSingleShaderBindings& ShaderBindings, float InDisplayGamma, float bInvertAlpha, float InContrast)
 	{
 		FVector4f Values( 2.2f / InDisplayGamma, 1.0f/InDisplayGamma, bInvertAlpha, InContrast);
 
-		SetShaderValue(BatchedParameters, GammaAndAlphaValues, Values);
+		ShaderBindings.Add(GammaAndAlphaValues, Values);
 	}
 
 private:
-	
-	/** Texture parameter used by the shader */
 	LAYOUT_FIELD(FShaderResourceParameter, TextureParameter);
 	LAYOUT_FIELD(FShaderResourceParameter, TextureParameterSampler);
 	LAYOUT_FIELD(FShaderResourceParameter, InPageTableTexture);
-	LAYOUT_FIELD(FShaderParameter, VTPackedPageTableUniform);
+	LAYOUT_FIELD(FShaderParameter, VTPackedPageTableUniform0);
+	LAYOUT_FIELD(FShaderParameter, VTPackedPageTableUniform1);
 	LAYOUT_FIELD(FShaderParameter, VTPackedUniform);
 	LAYOUT_FIELD(FShaderParameter, ShaderParams);
 	LAYOUT_FIELD(FShaderParameter, ShaderParams2);
@@ -306,9 +271,9 @@ public:
 	*
 	* @param InShaderParams Shader params to use
 	*/
-	void SetBatchColor(FRHIBatchedShaderParameters& BatchedParameters, const FLinearColor& InBatchColor)
+	void SetBatchColor(FMeshDrawSingleShaderBindings& ShaderBindings, const FLinearColor& InBatchColor)
 	{
-		SetShaderValue(BatchedParameters, BatchColor, InBatchColor);
+		ShaderBindings.Add(BatchColor, InBatchColor);
 	}
 
 private:
