Moved slate viewport present to occur after rendering all slate windows, rather than pipelining them. This allows the ReadyToPresentDelegate to be broadcast at the correct time.

diff --git a/Engine/Source/Runtime/SlateCore/Public/Rendering/SlateRenderer.h b/Engine/Source/Runtime/SlateCore/Public/Rendering/SlateRenderer.h
--- a/Engine/Source/Runtime/SlateCore/Public/Rendering/SlateRenderer.h
+++ b/Engine/Source/Runtime/SlateCore/Public/Rendering/SlateRenderer.h
@@ -14,6 +14,8 @@
 
 class FRHITexture;
 class FRenderTarget;
+class FRDGBuilder;
+class FRDGTexture;
 class FSlateDrawBuffer;
 class FSlateUpdatableTexture;
 class ISlate3DRenderer;
@@ -296,10 +298,14 @@ public:
 	DECLARE_MULTICAST_DELEGATE_OneParam(FOnPostResizeWindowBackbuffer, void*);
 	FOnPostResizeWindowBackbuffer& OnPostResizeWindowBackBuffer() { return PostResizeBackBufferDelegate; }
 
-	/** Callback on the render thread after slate rendering finishes and right before present is called */
+	/** Callback on the render thread after slate rendering finishes and right before present is called. */
 	DECLARE_TS_MULTICAST_DELEGATE_TwoParams(FOnBackBufferReadyToPresent, SWindow&, const FTextureRHIRef&);
 	FOnBackBufferReadyToPresent& OnBackBufferReadyToPresent() { return OnBackBufferReadyToPresentDelegate; }
 
+	/** Callback on the render thread after slate rendering finishes and right before present is called. */
+	DECLARE_TS_MULTICAST_DELEGATE_ThreeParams(FOnAddBackBufferReadyToPresentPass, FRDGBuilder&, SWindow&, FRDGTexture*);
+	FOnAddBackBufferReadyToPresentPass& OnAddBackBufferReadyToPresentPass() { return OnAddBackBufferReadyToPresentPassDelegate; }
+
 	/** 
 	 * Sets which color vision filter to use
 	 */
@@ -583,6 +589,7 @@ protected:
 	FOnPostResizeWindowBackbuffer PostResizeBackBufferDelegate;
 
 	FOnBackBufferReadyToPresent OnBackBufferReadyToPresentDelegate;
+	FOnAddBackBufferReadyToPresentPass OnAddBackBufferReadyToPresentPassDelegate;
 
 	/**
 	 * Necessary to grab before flushing the resource pool, as it may be being 
diff --git a/Engine/Source/Runtime/SlateRHIRenderer/Private/SlateRHIRenderer.cpp b/Engine/Source/Runtime/SlateRHIRenderer/Private/SlateRHIRenderer.cpp
--- a/Engine/Source/Runtime/SlateRHIRenderer/Private/SlateRHIRenderer.cpp
+++ b/Engine/Source/Runtime/SlateRHIRenderer/Private/SlateRHIRenderer.cpp
@@ -202,6 +202,13 @@ struct FSlateDrawWindowPassInputs
 	bool bClear = false;
 };
 
+struct FSlateDrawWindowPassOutputs
+{
+	FRHIViewport* ViewportRHI = nullptr;
+	FRHITexture* ViewportTextureRHI = nullptr;
+	FRHITexture* OutputTextureRHI = nullptr;
+};
+
 FMatrix CreateSlateProjectionMatrix(uint32 Width, uint32 Height)
 {
 	// Create ortho projection matrix
@@ -647,7 +654,7 @@ public:
 
 IMPLEMENT_GLOBAL_SHADER(FCompositeCS, "/Engine/Private/CompositeUIPixelShader.usf", "CompositeUICS", SF_Compute);
 
-void FSlateRHIRenderer::DrawWindow_RenderThread(FRDGBuilder& GraphBuilder, const FSlateDrawWindowPassInputs& Inputs)
+FSlateDrawWindowPassOutputs FSlateRHIRenderer::DrawWindow_RenderThread(FRDGBuilder& GraphBuilder, const FSlateDrawWindowPassInputs& Inputs)
 {
 	LLM_SCOPE(ELLMTag::SceneRender);
 
@@ -655,12 +662,10 @@ void FSlateRHIRenderer::DrawWindow_RenderThread(FRDGBuilder& GraphBuilder, const
 	FSlateWindowElementList& WindowElementList = *Inputs.WindowElementList;
 
 	FMaterialRenderProxy::UpdateDeferredCachedUniformExpressions(GraphBuilder.RHICmdList);
+	GetRendererModule().InitializeSystemTextures(GraphBuilder.RHICmdList);
 
-	AddPass(GraphBuilder, RDG_EVENT_NAME("BeginDrawingViewport"), [ViewportRHI = ViewportInfo.ViewportRHI] (FRHICommandListImmediate& RHICmdList)
-	{
-		GetRendererModule().InitializeSystemTextures(RHICmdList);
-		RHICmdList.BeginDrawingViewport(ViewportRHI, FTextureRHIRef());
-	});
+	FRHITexture* ViewportTextureRHI;
+	FRHITexture* OutputTextureRHI;
 
 	{
 		RDG_GPU_MASK_SCOPE(GraphBuilder, FRHIGPUMask::FromIndex(RHIGetViewportNextPresentGPUIndex(ViewportInfo.ViewportRHI)));
@@ -675,7 +680,7 @@ void FSlateRHIRenderer::DrawWindow_RenderThread(FRDGBuilder& GraphBuilder, const
 		TRACE_CPUPROFILER_EVENT_SCOPE(Slate::DrawWindow_RenderThread);
 
 		// The viewport texture is an optional user-allocated render target. This is rendered to if valid.
-		FRHITexture* ViewportTextureRHI = ViewportInfo.GetRenderTargetTexture();
+		ViewportTextureRHI = ViewportInfo.GetRenderTargetTexture();
 
 		// The swap chain is the final output. This is rendered to if no viewport render target is provided.
 		FRHITexture* SwapChainTextureRHI = RHIGetViewportBackBuffer(ViewportInfo.ViewportRHI);
@@ -684,7 +689,7 @@ void FSlateRHIRenderer::DrawWindow_RenderThread(FRDGBuilder& GraphBuilder, const
 		const bool bCompositeStereoToSwapChain = ViewportTextureRHI && GEngine && GEngine->StereoRenderingDevice.IsValid();
 
 		// The output texture is what we ultimately render or composite slate elements into.
-		FRHITexture* OutputTextureRHI = bCompositeStereoToSwapChain ? ViewportTextureRHI : SwapChainTextureRHI;
+		OutputTextureRHI = bCompositeStereoToSwapChain ? ViewportTextureRHI : SwapChainTextureRHI;
 		FRDGTexture* OutputTexture = RegisterExternalTexture(GraphBuilder, OutputTextureRHI, TEXT("SlateOutputTexture"));
 
 		// The elements texture contains UI elements. It can be the same as the output or allocated separately and composited.
@@ -872,17 +877,16 @@ void FSlateRHIRenderer::DrawWindow_RenderThread(FRDGBuilder& GraphBuilder, const
 				AddDrawScreenPass(GraphBuilder, RDG_EVENT_NAME("CompositeUI"), FScreenPassViewInfo(), Viewport, Viewport, PixelShader, PassParameters);
 			}
 		}
-
-		AddPass(GraphBuilder, RDG_EVENT_NAME("ReadyToPresent"), [this, Window = Inputs.Window, WindowSize = WindowElementList.GetWindowSize(), bCompositeStereoToSwapChain, SwapChainTextureRHI, ViewportTextureRHI, OutputTextureRHI] (FRHICommandListImmediate& RHICmdList)
+	
+		if (bCompositeStereoToSwapChain)
 		{
-			if (bCompositeStereoToSwapChain)
+			AddPass(GraphBuilder, RDG_EVENT_NAME("CompositeStereoToSwapChain"), [this, WindowSize = WindowElementList.GetWindowSize(), SwapChainTextureRHI, ViewportTextureRHI] (FRHICommandListImmediate& RHICmdList)
 			{
 				GEngine->StereoRenderingDevice->RenderTexture_RenderThread(RHICmdList, SwapChainTextureRHI, ViewportTextureRHI, WindowSize);
-			}
-			
-			// Fire delegate to inform bound functions the back buffer is ready to be captured.
-			OnBackBufferReadyToPresentDelegate.Broadcast(*Window, OutputTextureRHI);
-		});
+			});
+		}
+
+		OnAddBackBufferReadyToPresentPassDelegate.Broadcast(GraphBuilder, *Inputs.Window, OutputTexture);
 
 		if (ScreenshotState.ViewportToCapture == &ViewportInfo)
 		{
@@ -916,102 +920,111 @@ void FSlateRHIRenderer::DrawWindow_RenderThread(FRDGBuilder& GraphBuilder, const
 		}
 	}
 
-	AddPass(GraphBuilder, RDG_EVENT_NAME("EndDrawingViewport"), [this, ViewportRHI = ViewportInfo.ViewportRHI, bLockToVsync = Inputs.bLockToVsync] (FRHICommandListImmediate& RHICmdList)
-	{
-		uint32 StartTime = FPlatformTime::Cycles();
+	FSlateDrawWindowPassOutputs Outputs;
+	Outputs.ViewportRHI = ViewportInfo.ViewportRHI;
+	Outputs.ViewportTextureRHI = ViewportTextureRHI;
+	Outputs.OutputTextureRHI = OutputTextureRHI;
+	return Outputs;
+}
 
-		RHICmdList.EnqueueLambda([CurrentFrameCounter = GFrameCounterRenderThread](FRHICommandListImmediate& InRHICmdList)
-		{
-			UEngine::SetPresentLatencyMarkerStart(CurrentFrameCounter);
-		});
+void FSlateRHIRenderer::PresentWindow_RenderThread(FRHICommandListImmediate& RHICmdList, const FSlateDrawWindowPassInputs& DrawPassInputs, const FSlateDrawWindowPassOutputs& DrawPassOutputs)
+{
+	OnBackBufferReadyToPresentDelegate.Broadcast(*DrawPassInputs.Window, DrawPassOutputs.OutputTextureRHI);
 
-		RHICmdList.EndDrawingViewport(ViewportRHI, true, bLockToVsync);
+	uint32 StartTime = FPlatformTime::Cycles();
 
-		RHICmdList.EnqueueLambda([CurrentFrameCounter = GFrameCounterRenderThread](FRHICommandListImmediate& InRHICmdList)
-		{
-			UEngine::SetPresentLatencyMarkerEnd(CurrentFrameCounter);
-		});
+	RHICmdList.EnqueueLambda([CurrentFrameCounter = GFrameCounterRenderThread](FRHICommandListImmediate& InRHICmdList)
+	{
+		UEngine::SetPresentLatencyMarkerStart(CurrentFrameCounter);
+	});
 
-		uint32 EndTime = FPlatformTime::Cycles();
+	RHICmdList.BeginDrawingViewport(DrawPassOutputs.ViewportRHI, FTextureRHIRef());
+	RHICmdList.EndDrawingViewport(DrawPassOutputs.ViewportRHI, true, DrawPassInputs.bLockToVsync);
 
-		GSwapBufferTime = EndTime - StartTime;
-		SET_CYCLE_COUNTER(STAT_PresentTime, GSwapBufferTime);
+	RHICmdList.EnqueueLambda([CurrentFrameCounter = GFrameCounterRenderThread](FRHICommandListImmediate& InRHICmdList)
+	{
+		UEngine::SetPresentLatencyMarkerEnd(CurrentFrameCounter);
+	});
 
-		static uint32 LastTimestamp = FPlatformTime::Cycles();
-		uint32 ThreadTime = EndTime - LastTimestamp;
-		LastTimestamp = EndTime;
+	uint32 EndTime = FPlatformTime::Cycles();
 
-		uint32 RenderThreadIdle = 0;
+	GSwapBufferTime = EndTime - StartTime;
+	SET_CYCLE_COUNTER(STAT_PresentTime, GSwapBufferTime);
 
-		FThreadIdleStats& RenderThread = FThreadIdleStats::Get();
-		GRenderThreadIdle[ERenderThreadIdleTypes::WaitingForAllOtherSleep] = RenderThread.Waits;
-		GRenderThreadIdle[ERenderThreadIdleTypes::WaitingForGPUPresent] += GSwapBufferTime;
+	static uint32 LastTimestamp = FPlatformTime::Cycles();
+	uint32 ThreadTime = EndTime - LastTimestamp;
+	LastTimestamp = EndTime;
 
-		SET_CYCLE_COUNTER(STAT_RenderingIdleTime_RenderThreadSleepTime, GRenderThreadIdle[ERenderThreadIdleTypes::WaitingForAllOtherSleep]);
-		SET_CYCLE_COUNTER(STAT_RenderingIdleTime_WaitingForGPUQuery   , GRenderThreadIdle[ERenderThreadIdleTypes::WaitingForGPUQuery     ]);
-		SET_CYCLE_COUNTER(STAT_RenderingIdleTime_WaitingForGPUPresent , GRenderThreadIdle[ERenderThreadIdleTypes::WaitingForGPUPresent   ]);
+	uint32 RenderThreadIdle = 0;
 
-		const uint32 RenderThreadNonCriticalWaits = RenderThread.Waits - RenderThread.WaitsCriticalPath;
-		const uint32 RenderThreadWaitingForGPUQuery = GRenderThreadIdle[ERenderThreadIdleTypes::WaitingForGPUQuery];
+	FThreadIdleStats& RenderThread = FThreadIdleStats::Get();
+	GRenderThreadIdle[ERenderThreadIdleTypes::WaitingForAllOtherSleep] = RenderThread.Waits;
+	GRenderThreadIdle[ERenderThreadIdleTypes::WaitingForGPUPresent] += GSwapBufferTime;
 
-		// Set the RenderThreadIdle CSV stats
-		CSV_CUSTOM_STAT(RenderThreadIdle, Total          , FPlatformTime::ToMilliseconds(RenderThread.Waits            ), ECsvCustomStatOp::Set);
-		CSV_CUSTOM_STAT(RenderThreadIdle, CriticalPath   , FPlatformTime::ToMilliseconds(RenderThread.WaitsCriticalPath), ECsvCustomStatOp::Set);
-		CSV_CUSTOM_STAT(RenderThreadIdle, SwapBuffer     , FPlatformTime::ToMilliseconds(GSwapBufferTime               ), ECsvCustomStatOp::Set);
-		CSV_CUSTOM_STAT(RenderThreadIdle, NonCriticalPath, FPlatformTime::ToMilliseconds(RenderThreadNonCriticalWaits  ), ECsvCustomStatOp::Set);
-		CSV_CUSTOM_STAT(RenderThreadIdle, GPUQuery       , FPlatformTime::ToMilliseconds(RenderThreadWaitingForGPUQuery), ECsvCustomStatOp::Set);
+	SET_CYCLE_COUNTER(STAT_RenderingIdleTime_RenderThreadSleepTime, GRenderThreadIdle[ERenderThreadIdleTypes::WaitingForAllOtherSleep]);
+	SET_CYCLE_COUNTER(STAT_RenderingIdleTime_WaitingForGPUQuery   , GRenderThreadIdle[ERenderThreadIdleTypes::WaitingForGPUQuery     ]);
+	SET_CYCLE_COUNTER(STAT_RenderingIdleTime_WaitingForGPUPresent , GRenderThreadIdle[ERenderThreadIdleTypes::WaitingForGPUPresent   ]);
 
-		for (int32 Index = 0; Index < ERenderThreadIdleTypes::Num; Index++)
-		{
-			RenderThreadIdle += GRenderThreadIdle[Index];
-			GRenderThreadIdle[Index] = 0;
-		}
+	const uint32 RenderThreadNonCriticalWaits = RenderThread.Waits - RenderThread.WaitsCriticalPath;
+	const uint32 RenderThreadWaitingForGPUQuery = GRenderThreadIdle[ERenderThreadIdleTypes::WaitingForGPUQuery];
 
-		SET_CYCLE_COUNTER(STAT_RenderingIdleTime, RenderThreadIdle);
-		GRenderThreadTime = (ThreadTime > RenderThreadIdle) ? (ThreadTime - RenderThreadIdle) : ThreadTime;
-		GRenderThreadWaitTime = RenderThreadIdle;
+	// Set the RenderThreadIdle CSV stats
+	CSV_CUSTOM_STAT(RenderThreadIdle, Total          , FPlatformTime::ToMilliseconds(RenderThread.Waits            ), ECsvCustomStatOp::Set);
+	CSV_CUSTOM_STAT(RenderThreadIdle, CriticalPath   , FPlatformTime::ToMilliseconds(RenderThread.WaitsCriticalPath), ECsvCustomStatOp::Set);
+	CSV_CUSTOM_STAT(RenderThreadIdle, SwapBuffer     , FPlatformTime::ToMilliseconds(GSwapBufferTime               ), ECsvCustomStatOp::Set);
+	CSV_CUSTOM_STAT(RenderThreadIdle, NonCriticalPath, FPlatformTime::ToMilliseconds(RenderThreadNonCriticalWaits  ), ECsvCustomStatOp::Set);
+	CSV_CUSTOM_STAT(RenderThreadIdle, GPUQuery       , FPlatformTime::ToMilliseconds(RenderThreadWaitingForGPUQuery), ECsvCustomStatOp::Set);
 
-		// Compute GRenderThreadTimeCriticalPath
-		uint32 RenderThreadNonCriticalPathIdle = RenderThreadIdle - RenderThread.WaitsCriticalPath;
-		GRenderThreadTimeCriticalPath = (ThreadTime > RenderThreadNonCriticalPathIdle) ? (ThreadTime - RenderThreadNonCriticalPathIdle) : ThreadTime;
-		SET_CYCLE_COUNTER(STAT_RenderThreadCriticalPath, GRenderThreadTimeCriticalPath);
+	for (int32 Index = 0; Index < ERenderThreadIdleTypes::Num; Index++)
+	{
+		RenderThreadIdle += GRenderThreadIdle[Index];
+		GRenderThreadIdle[Index] = 0;
+	}
 
-		if (CVarRenderThreadTimeIncludesDependentWaits.GetValueOnRenderThread())
-		{
-			// Optionally force the renderthread stat to include dependent waits
-			GRenderThreadTime = GRenderThreadTimeCriticalPath;
-		}
+	SET_CYCLE_COUNTER(STAT_RenderingIdleTime, RenderThreadIdle);
+	GRenderThreadTime = (ThreadTime > RenderThreadIdle) ? (ThreadTime - RenderThreadIdle) : ThreadTime;
+	GRenderThreadWaitTime = RenderThreadIdle;
+
+	// Compute GRenderThreadTimeCriticalPath
+	uint32 RenderThreadNonCriticalPathIdle = RenderThreadIdle - RenderThread.WaitsCriticalPath;
+	GRenderThreadTimeCriticalPath = (ThreadTime > RenderThreadNonCriticalPathIdle) ? (ThreadTime - RenderThreadNonCriticalPathIdle) : ThreadTime;
+	SET_CYCLE_COUNTER(STAT_RenderThreadCriticalPath, GRenderThreadTimeCriticalPath);
 
-		// Reset the idle stats
-		RenderThread.Reset();
+	if (CVarRenderThreadTimeIncludesDependentWaits.GetValueOnRenderThread())
+	{
+		// Optionally force the renderthread stat to include dependent waits
+		GRenderThreadTime = GRenderThreadTimeCriticalPath;
+	}
+
+	// Reset the idle stats
+	RenderThread.Reset();
 	
-		static TOptional<uint32> RHITCycles;
-		if (IsRunningRHIInSeparateThread())
+	static TOptional<uint32> RHITCycles;
+	if (IsRunningRHIInSeparateThread())
+	{
+		RHICmdList.EnqueueLambda([](FRHICommandListImmediate&)
 		{
-			RHICmdList.EnqueueLambda([](FRHICommandListImmediate&)
-			{
-				// Update RHI thread time
-				FThreadIdleStats& RHIThreadStats = FThreadIdleStats::Get();
+			// Update RHI thread time
+			FThreadIdleStats& RHIThreadStats = FThreadIdleStats::Get();
 
-				if (!RHITCycles.IsSet())
-				{
-					RHITCycles = FPlatformTime::Cycles();
-				}
+			if (!RHITCycles.IsSet())
+			{
+				RHITCycles = FPlatformTime::Cycles();
+			}
 
-				uint32 Next = FPlatformTime::Cycles();
+			uint32 Next = FPlatformTime::Cycles();
 
-				int32 Result = int32(Next - RHITCycles.GetValue() - RHIThreadStats.Waits);
-				RHITCycles = Next;
+			int32 Result = int32(Next - RHITCycles.GetValue() - RHIThreadStats.Waits);
+			RHITCycles = Next;
 
-				FPlatformAtomics::AtomicStore((int32*)&GRHIThreadTime, FMath::Max(Result, 0));
-				RHIThreadStats.Reset();
-			});
-		}
-		else
-		{
-			RHITCycles.Reset();
-		}
-	});
+			FPlatformAtomics::AtomicStore((int32*)&GRHIThreadTime, FMath::Max(Result, 0));
+			RHIThreadStats.Reset();
+		});
+	}
+	else
+	{
+		RHITCycles.Reset();
+	}
 }
 
 void FSlateRHIRenderer::DrawWindows(FSlateDrawBuffer& WindowDrawBuffer)
@@ -1305,23 +1318,34 @@ void FSlateRHIRenderer::DrawWindows_Private(FSlateDrawBuffer& WindowDrawBuffer)
 	{
 		ENQUEUE_RENDER_COMMAND(SlateDrawWindowsCommand)([DrawWindowsCommand = MoveTemp(DrawWindowsCommand)](FRHICommandListImmediate& RHICmdList)
 		{
-			FRDGBuilder GraphBuilder(RHICmdList, RDG_EVENT_NAME("Slate"));
+			TArray<FSlateDrawWindowPassOutputs, FConcurrentLinearArrayAllocator> DrawWindowPassOutputs;
+			DrawWindowPassOutputs.Reserve(DrawWindowsCommand->Windows.Num());
 
-			for (const FRenderThreadUpdateContext& DeferredUpdateContext : DrawWindowsCommand->DeferredUpdates)
 			{
-				DeferredUpdateContext.Renderer->DrawWindowToTarget_RenderThread(GraphBuilder, DeferredUpdateContext);
-			}
+				FRDGBuilder GraphBuilder(RHICmdList, RDG_EVENT_NAME("Slate"));
 
-			for (const FSlateDrawWindowPassInputs& DrawWindowPassInputs : DrawWindowsCommand->Windows)
-			{
-				DrawWindowPassInputs.Renderer->DrawWindow_RenderThread(GraphBuilder, DrawWindowPassInputs);
-			}
+				for (const FRenderThreadUpdateContext& DeferredUpdateContext : DrawWindowsCommand->DeferredUpdates)
+				{
+					DeferredUpdateContext.Renderer->DrawWindowToTarget_RenderThread(GraphBuilder, DeferredUpdateContext);
+				}
 
-			GraphBuilder.Execute();
+				for (const FSlateDrawWindowPassInputs& DrawWindowPassInputs : DrawWindowsCommand->Windows)
+				{
+					DrawWindowPassOutputs.Emplace(DrawWindowPassInputs.Renderer->DrawWindow_RenderThread(GraphBuilder, DrawWindowPassInputs));
+				}
+
+				GraphBuilder.Execute();
+
+				for (const FRenderThreadUpdateContext& DeferredUpdateContext : DrawWindowsCommand->DeferredUpdates)
+				{
+					DeferredUpdateContext.Renderer->ReleaseDrawBuffer(*DeferredUpdateContext.WindowDrawBuffer);
+				}
+			}
 
-			for (const FRenderThreadUpdateContext& DeferredUpdateContext : DrawWindowsCommand->DeferredUpdates)
+			for (int32 WindowIndex = 0; WindowIndex < DrawWindowsCommand->Windows.Num(); ++WindowIndex)
 			{
-				DeferredUpdateContext.Renderer->ReleaseDrawBuffer(*DeferredUpdateContext.WindowDrawBuffer);
+				const FSlateDrawWindowPassInputs& DrawWindowPassInputs = DrawWindowsCommand->Windows[WindowIndex];
+				DrawWindowPassInputs.Renderer->PresentWindow_RenderThread(RHICmdList, DrawWindowPassInputs, DrawWindowPassOutputs[WindowIndex]);
 			}
 		});
 
diff --git a/Engine/Source/Runtime/SlateRHIRenderer/Private/SlateRHIRenderer.h b/Engine/Source/Runtime/SlateRHIRenderer/Private/SlateRHIRenderer.h
--- a/Engine/Source/Runtime/SlateRHIRenderer/Private/SlateRHIRenderer.h
+++ b/Engine/Source/Runtime/SlateRHIRenderer/Private/SlateRHIRenderer.h
@@ -36,6 +36,7 @@ const uint32 NumDrawBuffers = 3;
 struct FFastPathRenderingDataCleanUpList;
 
 struct FSlateDrawWindowPassInputs;
+struct FSlateDrawWindowPassOutputs;
 struct FSlateViewportInfo;
 struct FSlatePostProcessUpdateRequest;
 
@@ -109,7 +110,9 @@ public:
 	virtual void SetWindowRenderTarget(const SWindow& Window, class IViewportRenderTargetProvider* Provider) override;
 
 private:
-	void DrawWindow_RenderThread(FRDGBuilder& GraphBuilder, const FSlateDrawWindowPassInputs& Inputs);
+	FSlateDrawWindowPassOutputs DrawWindow_RenderThread(FRDGBuilder& GraphBuilder, const FSlateDrawWindowPassInputs& Inputs);
+
+	void PresentWindow_RenderThread(FRHICommandListImmediate& RHICmdList, const FSlateDrawWindowPassInputs& DrawPassInputs, const FSlateDrawWindowPassOutputs& DrawPassOutputs);
 
 	/** Loads all known textures from Slate styles */
 	void LoadUsedTextures();
