Fixed crash when reloading texture being used as an atlas while it's sprite is visible

diff --git a/Engine/Source/Runtime/SlateRHIRenderer/Private/SlateAtlasedTextureResource.cpp b/Engine/Source/Runtime/SlateRHIRenderer/Private/SlateAtlasedTextureResource.cpp
--- a/Engine/Source/Runtime/SlateRHIRenderer/Private/SlateAtlasedTextureResource.cpp
+++ b/Engine/Source/Runtime/SlateRHIRenderer/Private/SlateAtlasedTextureResource.cpp
@@ -17,6 +17,11 @@ FSlateAtlasedTextureResource::~FSlateAtlasedTextureResource()
 	}
 }
 
+bool FSlateAtlasedTextureResource::IsResourceValid() const
+{
+	return TextureObject && TextureObject->IsValidLowLevel();
+}
+
 FSlateShaderResourceProxy* FSlateAtlasedTextureResource::FindOrCreateAtlasedProxy(UObject* InAtlasedObject, const FSlateAtlasData& AtlasData)
 {
 	FSlateShaderResourceProxy* Proxy = ProxyMap.FindRef(InAtlasedObject);
diff --git a/Engine/Source/Runtime/SlateRHIRenderer/Private/SlateAtlasedTextureResource.h b/Engine/Source/Runtime/SlateRHIRenderer/Private/SlateAtlasedTextureResource.h
--- a/Engine/Source/Runtime/SlateRHIRenderer/Private/SlateAtlasedTextureResource.h
+++ b/Engine/Source/Runtime/SlateRHIRenderer/Private/SlateAtlasedTextureResource.h
@@ -19,6 +19,10 @@ public:
 	
 	/** Destructor */
 	virtual ~FSlateAtlasedTextureResource();
+
+	// Begin FSlateShaderResource Interface
+	virtual bool IsResourceValid() const override;
+	// End FSlateShaderResource Interface
 	
 	/**
 	 * Finds or creates the rendering proxy for a given atlas'ed object.
