Fixed slate instanced render batch to avoid touching proxy inside of RDG pass.

diff --git a/Engine/Source/Runtime/RenderCore/Private/RenderGraphPrivate.cpp b/Engine/Source/Runtime/RenderCore/Private/RenderGraphPrivate.cpp
--- a/Engine/Source/Runtime/RenderCore/Private/RenderGraphPrivate.cpp
+++ b/Engine/Source/Runtime/RenderCore/Private/RenderGraphPrivate.cpp
@@ -342,7 +342,7 @@ FAutoConsoleVariableRef CVarRDGParallelSetup(
 	TEXT(" 1: pass setup is done asynchronously (default);"),
 	ECVF_RenderThreadSafe);
 
-int32 GRDGParallelExecute = 1;
+int32 GRDGParallelExecute = 2;
 FAutoConsoleVariableRef CVarRDGParallelExecute(
 	TEXT("r.RDG.ParallelExecute"), GRDGParallelExecute,
 	TEXT("Whether to enable parallel execution of passes when supported.")
@@ -378,7 +378,7 @@ FAutoConsoleVariableRef CVarRDGParallelExecutePassMax(
 	TEXT("The maximum span of contiguous passes eligible for parallel execution for the span to be offloaded to a task."),
 	ECVF_RenderThreadSafe);
 
-int32 GRDGParallelExecutePassTaskModeThreshold = 0;
+int32 GRDGParallelExecutePassTaskModeThreshold = 2;
 FAutoConsoleVariableRef CVarRDGParallelExecutePassTaskModeThreshold(
 	TEXT("r.RDG.ParallelExecute.PassTaskModeThreshold"), GRDGParallelExecutePassTaskModeThreshold,
 	TEXT(" 0: A pass that is not marked async will mark the entire parallel pass set as awaited.")
diff --git a/Engine/Source/Runtime/SlateCore/Public/Rendering/RenderingCommon.h b/Engine/Source/Runtime/SlateCore/Public/Rendering/RenderingCommon.h
--- a/Engine/Source/Runtime/SlateCore/Public/Rendering/RenderingCommon.h
+++ b/Engine/Source/Runtime/SlateCore/Public/Rendering/RenderingCommon.h
@@ -20,6 +20,7 @@
 #include "RenderingCommon.generated.h"
 
 class FRHICommandListImmediate;
+class FRHIBuffer;
 class FRDGBuilder;
 class FRDGTexture;
 class FSlateElementBatcher;
@@ -942,7 +943,10 @@ protected:
 	virtual ~ISlateUpdatableInstanceBufferRenderProxy() {};
 
 public:
-	virtual void BindStreamSource(class FRHICommandList& RHICmdList, int32 StreamIndex, uint32 InstanceOffset) = 0;
+	virtual FRHIBuffer* GetRHI() const = 0;
+
+	UE_DEPRECATED(5.5, "Use GetStreamSource instead")
+	virtual void BindStreamSource(class FRHICommandList& RHICmdList, int32 StreamIndex, uint32 InstanceOffset) {}
 };
 
 typedef TArray<FVector4f> FSlateInstanceBufferData;
diff --git a/Engine/Source/Runtime/SlateRHIRenderer/Private/SlateRHIRenderingPolicy.cpp b/Engine/Source/Runtime/SlateRHIRenderer/Private/SlateRHIRenderingPolicy.cpp
--- a/Engine/Source/Runtime/SlateRHIRenderer/Private/SlateRHIRenderingPolicy.cpp
+++ b/Engine/Source/Runtime/SlateRHIRenderer/Private/SlateRHIRenderingPolicy.cpp
@@ -986,6 +986,7 @@ struct FSlateRenderBatchOp
 	const FSlateClippingOp* ClippingStateOp;
 	FSlateDrawShaderBindings* VertexBindings;
 	FSlateDrawShaderBindings* PixelBindings;
+	FRHIBuffer* InstanceBuffer;
 	FRHIBlendState* BlendState;
 	ESlateShaderResource::Type ShaderResourceType;
 };
@@ -1237,6 +1238,7 @@ FSlateRenderBatchOp* CreateSlateRenderBatchOp(
 	RenderBatchOp->ShaderResourceType = ResourceType;
 	RenderBatchOp->VertexBindings     = VertexBindings;
 	RenderBatchOp->PixelBindings      = PixelBindings;
+	RenderBatchOp->InstanceBuffer     = bUseInstancing ? RenderBatch->InstanceData->GetRHI() : nullptr;
 	RenderBatchOp->BlendState         = BlendState;
 	RenderBatchOp->Next               = nullptr;
 	return RenderBatchOp;
@@ -1285,9 +1287,7 @@ void DrawSlateRenderBatch(
 
 	if (RenderBatchOp.ShaderResourceType == ESlateShaderResource::Material)
 	{
-		const bool bUseInstancing = RenderBatch.InstanceCount > 0 && RenderBatch.InstanceData != nullptr;
-
-		State.GraphicsPSOInit.BoundShaderState.VertexDeclarationRHI = bUseInstancing ? GSlateInstancedVertexDeclaration.VertexDeclarationRHI : GSlateVertexDeclaration.VertexDeclarationRHI;
+		State.GraphicsPSOInit.BoundShaderState.VertexDeclarationRHI = RenderBatchOp.InstanceBuffer ? GSlateInstancedVertexDeclaration.VertexDeclarationRHI : GSlateVertexDeclaration.VertexDeclarationRHI;
 		State.GraphicsPSOInit.BoundShaderState.VertexShaderRHI = RenderBatchOp.VertexBindings->Shader.GetVertexShader();
 		State.GraphicsPSOInit.BoundShaderState.PixelShaderRHI  = RenderBatchOp.PixelBindings->Shader.GetPixelShader();
 		State.GraphicsPSOInit.PrimitiveType = GetRHIPrimitiveType(RenderBatch.DrawPrimitiveType);
@@ -1297,16 +1297,15 @@ void DrawSlateRenderBatch(
 		RenderBatchOp.VertexBindings->SetOnCommandList(RHICmdList);
 		RenderBatchOp.PixelBindings->SetOnCommandList(RHICmdList);
 
-		if (bUseInstancing)
-		{
-			RenderBatch.InstanceData->BindStreamSource(RHICmdList, 1, RenderBatch.InstanceOffset);
+		RHICmdList.SetStreamSource(0, ElementsVertexBuffer, RenderBatch.VertexOffset * sizeof(FSlateVertex));
 
-			RHICmdList.SetStreamSource(0, ElementsVertexBuffer, RenderBatch.VertexOffset * sizeof(FSlateVertex));
+		if (RenderBatchOp.InstanceBuffer)
+		{
+			RHICmdList.SetStreamSource(1, RenderBatchOp.InstanceBuffer, RenderBatch.InstanceOffset * sizeof(FSlateInstanceBufferData::ElementType));
 			RHICmdList.DrawIndexedPrimitive(ElementsIndexBuffer, 0, 0, RenderBatch.NumVertices, RenderBatch.IndexOffset, PrimitiveCount, RenderBatch.InstanceCount);
 		}
 		else
 		{
-			RHICmdList.SetStreamSource(0, ElementsVertexBuffer, RenderBatch.VertexOffset * sizeof(FSlateVertex));
 			RHICmdList.SetStreamSource(1, nullptr, 0);
 			RHICmdList.DrawIndexedPrimitive(ElementsIndexBuffer, 0, 0, RenderBatch.NumVertices, RenderBatch.IndexOffset, PrimitiveCount, 1);
 		}
diff --git a/Engine/Source/Runtime/SlateRHIRenderer/Private/SlateUpdatableBuffer.cpp b/Engine/Source/Runtime/SlateRHIRenderer/Private/SlateUpdatableBuffer.cpp
--- a/Engine/Source/Runtime/SlateRHIRenderer/Private/SlateUpdatableBuffer.cpp
+++ b/Engine/Source/Runtime/SlateRHIRenderer/Private/SlateUpdatableBuffer.cpp
@@ -50,9 +50,4 @@ void FSlateUpdatableInstanceBuffer::FRenderProxy::Update(FRHICommandListImmediat
 	FMemory::Memcpy(InstanceBufferData, Data.GetData(), RequiredVertexBufferSize);
 	
 	RHICmdList.UnlockBuffer(InstanceBufferResource.VertexBufferRHI);
-}
-
-void FSlateUpdatableInstanceBuffer::FRenderProxy::BindStreamSource(FRHICommandList& RHICmdList, int32 StreamIndex, uint32 InstanceOffset)
-{
-	RHICmdList.SetStreamSource(StreamIndex, InstanceBufferResource.VertexBufferRHI, InstanceOffset * sizeof(FVector4f));
-}
+}
\ No newline at end of file
diff --git a/Engine/Source/Runtime/SlateRHIRenderer/Private/SlateUpdatableBuffer.h b/Engine/Source/Runtime/SlateRHIRenderer/Private/SlateUpdatableBuffer.h
--- a/Engine/Source/Runtime/SlateRHIRenderer/Private/SlateUpdatableBuffer.h
+++ b/Engine/Source/Runtime/SlateRHIRenderer/Private/SlateUpdatableBuffer.h
@@ -25,7 +25,10 @@ class FSlateUpdatableInstanceBuffer final : public ISlateUpdatableInstanceBuffer
 
 		void Update(FRHICommandListImmediate& RHICmdList, FSlateInstanceBufferData& Data);
 
-		virtual void BindStreamSource(FRHICommandList& RHICmdList, int32 StreamIndex, uint32 InstanceOffset) override final;
+		FRHIBuffer* GetRHI() const override
+		{
+			return InstanceBufferResource.VertexBufferRHI;
+		}
 	} *Proxy;
 
 public:
