Fixed slate upscale pass from having scissor overwritten by DrawScreenPass.

diff --git a/Engine/Source/Runtime/SlateRHIRenderer/Private/SlatePostProcessor.cpp b/Engine/Source/Runtime/SlateRHIRenderer/Private/SlatePostProcessor.cpp
--- a/Engine/Source/Runtime/SlateRHIRenderer/Private/SlatePostProcessor.cpp
+++ b/Engine/Source/Runtime/SlateRHIRenderer/Private/SlatePostProcessor.cpp
@@ -264,13 +264,14 @@ void AddSlatePostProcessUpsamplePass(FRDGBuilder& GraphBuilder, const FSlatePost
 		ERDGPassFlags::Raster,
 		[OutputViewport, InputViewport, PipelineState, PixelShader, ClippingOp = Inputs.ClippingOp, PassParameters](FRDGAsyncTask, FRHICommandList& RHICmdList)
 	{
+		RHICmdList.SetViewport(OutputViewport.Rect.Min.X, OutputViewport.Rect.Min.Y, 0.0f, OutputViewport.Rect.Max.X, OutputViewport.Rect.Max.Y, 1.0f);
+
 		// Stencil clipping will issue its own draw calls.
 		SetSlateClipping(RHICmdList, ClippingOp, OutputViewport.Rect);
 
-		DrawScreenPass(RHICmdList, FScreenPassViewInfo(), OutputViewport, InputViewport, PipelineState, EScreenPassDrawFlags::None, [&](FRHICommandList&)
-		{
-			SetShaderParameters(RHICmdList, PixelShader, PixelShader.GetPixelShader(), *PassParameters);
-		});
+		SetScreenPassPipelineState(RHICmdList, PipelineState);
+		SetShaderParameters(RHICmdList, PixelShader, PixelShader.GetPixelShader(), *PassParameters);
+		DrawScreenPass_PostSetup(RHICmdList, FScreenPassViewInfo(), OutputViewport, InputViewport, PipelineState, EScreenPassDrawFlags::None);
 	});
 }
 
