Fixed Slate post processing upsample using incorrect stencil vertices. Slate elements assume the viewport is the full window texture size, but the post processing upsample element was using the viewport for the post process element instead.

diff --git a/Engine/Source/Runtime/SlateRHIRenderer/Private/SlatePostProcessor.cpp b/Engine/Source/Runtime/SlateRHIRenderer/Private/SlatePostProcessor.cpp
--- a/Engine/Source/Runtime/SlateRHIRenderer/Private/SlatePostProcessor.cpp
+++ b/Engine/Source/Runtime/SlateRHIRenderer/Private/SlatePostProcessor.cpp
@@ -206,6 +206,7 @@ struct FSlatePostProcessUpsampleInputs
 
 	const FSlateClippingOp* ClippingOp = nullptr;
 	const FDepthStencilBinding* ClippingStencilBinding = nullptr;
+	FIntRect ClippingElementsViewRect;
 
 	FIntRect OutputRect;
 	FVector4f CornerRadius = FVector4f::Zero();
@@ -262,13 +263,16 @@ void AddSlatePostProcessUpsamplePass(FRDGBuilder& GraphBuilder, const FSlatePost
 		RDG_EVENT_NAME("Upsample"),
 		PassParameters,
 		ERDGPassFlags::Raster,
-		[OutputViewport, InputViewport, PipelineState, PixelShader, ClippingOp = Inputs.ClippingOp, PassParameters](FRDGAsyncTask, FRHICommandList& RHICmdList)
+		[OutputViewport, InputViewport, ClippingElementsViewRect = Inputs.ClippingElementsViewRect, PipelineState, PixelShader, ClippingOp = Inputs.ClippingOp, PassParameters](FRDGAsyncTask, FRHICommandList& RHICmdList)
 	{
-		RHICmdList.SetViewport(OutputViewport.Rect.Min.X, OutputViewport.Rect.Min.Y, 0.0f, OutputViewport.Rect.Max.X, OutputViewport.Rect.Max.Y, 1.0f);
-
-		// Stencil clipping will issue its own draw calls.
-		SetSlateClipping(RHICmdList, ClippingOp, OutputViewport.Rect);
+		if (ClippingOp)
+		{
+			RHICmdList.SetViewport(ClippingElementsViewRect.Min.X, ClippingElementsViewRect.Min.Y, 0.0f, ClippingElementsViewRect.Max.X, ClippingElementsViewRect.Max.Y, 1.0f);
+			// Stencil clipping will issue its own draw calls.
+			SetSlateClipping(RHICmdList, ClippingOp, ClippingElementsViewRect);
+		}
 
+		RHICmdList.SetViewport(OutputViewport.Rect.Min.X, OutputViewport.Rect.Min.Y, 0.0f, OutputViewport.Rect.Max.X, OutputViewport.Rect.Max.Y, 1.0f);
 		SetScreenPassPipelineState(RHICmdList, PipelineState);
 		SetShaderParameters(RHICmdList, PixelShader, PixelShader.GetPixelShader(), *PassParameters);
 		DrawScreenPass_PostSetup(RHICmdList, FScreenPassViewInfo(), OutputViewport, InputViewport, PipelineState, EScreenPassDrawFlags::None);
@@ -427,6 +431,7 @@ void AddSlatePostProcessBlurPass(FRDGBuilder& GraphBuilder, const FSlatePostProc
 	UpsampleInputs.OutputRect = Inputs.OutputRect;
 	UpsampleInputs.ClippingOp = Inputs.ClippingOp;
 	UpsampleInputs.ClippingStencilBinding = Inputs.ClippingStencilBinding;
+	UpsampleInputs.ClippingElementsViewRect = Inputs.ClippingElementsViewRect;
 	UpsampleInputs.CornerRadius = Inputs.CornerRadius;
 
 	AddSlatePostProcessUpsamplePass(GraphBuilder, UpsampleInputs);
diff --git a/Engine/Source/Runtime/SlateRHIRenderer/Private/SlatePostProcessor.h b/Engine/Source/Runtime/SlateRHIRenderer/Private/SlatePostProcessor.h
--- a/Engine/Source/Runtime/SlateRHIRenderer/Private/SlatePostProcessor.h
+++ b/Engine/Source/Runtime/SlateRHIRenderer/Private/SlatePostProcessor.h
@@ -58,8 +58,10 @@ struct FSlatePostProcessBlurPassInputs
 	FRDGTexture* OutputTexture = nullptr;
 	ERenderTargetLoadAction OutputLoadAction = ERenderTargetLoadAction::ELoad;
 
+	// An optional set of inputs for when a blur is performed as part of a slate render batch.
 	const FSlateClippingOp* ClippingOp = nullptr;
 	const FDepthStencilBinding* ClippingStencilBinding = nullptr;
+	FIntRect ClippingElementsViewRect;
 
 	FIntRect InputRect;
 	FIntRect OutputRect;
diff --git a/Engine/Source/Runtime/SlateRHIRenderer/Private/SlateRHIRenderingPolicy.cpp b/Engine/Source/Runtime/SlateRHIRenderer/Private/SlateRHIRenderingPolicy.cpp
--- a/Engine/Source/Runtime/SlateRHIRenderer/Private/SlateRHIRenderingPolicy.cpp
+++ b/Engine/Source/Runtime/SlateRHIRenderer/Private/SlateRHIRenderingPolicy.cpp
@@ -1591,6 +1591,7 @@ void AddSlateDrawElementsPass(
 			BlurInputs.OutputLoadAction = ConsumeLoadAction(ElementsLoadAction);
 			BlurInputs.ClippingOp = NextClippingOp;
 			BlurInputs.ClippingStencilBinding = &NextPassParameters->RenderTargets.DepthStencil;
+			BlurInputs.ClippingElementsViewRect = RenderBatchDrawInputs->ElementsViewRect;
 			BlurInputs.KernelSize = ShaderParams.PixelParams2.X;
 			BlurInputs.Strength = ShaderParams.PixelParams2.Y;
 			BlurInputs.DownsampleAmount = ShaderParams.PixelParams2.Z;
