Fixed crash when changing preview platform. Use the shader level max instead of the feature level max global shader map.

diff --git a/Engine/Source/Runtime/SlateRHIRenderer/Private/SlateRHIRenderingPolicy.cpp b/Engine/Source/Runtime/SlateRHIRenderer/Private/SlateRHIRenderingPolicy.cpp
--- a/Engine/Source/Runtime/SlateRHIRenderer/Private/SlateRHIRenderingPolicy.cpp
+++ b/Engine/Source/Runtime/SlateRHIRenderer/Private/SlateRHIRenderingPolicy.cpp
@@ -1365,7 +1365,7 @@ void AddSlateDrawElementsPass(
 	const float EngineGamma  = GEngine ? GEngine->GetDisplayGamma() : 2.2f;
 	const float DisplayGamma = Inputs.bAllowGammaCorrection && !Inputs.bElementsTextureIsHDRDisplay ? EngineGamma : 1.0f;
 
-	FGlobalShaderMap* ShaderMap = GetGlobalShaderMap(GMaxRHIFeatureLevel);
+	FGlobalShaderMap* ShaderMap = GetGlobalShaderMap(GMaxRHIShaderPlatform);
 
 	FSlateRHIResourceManager& ResourceManager = RenderingPolicy.GetResourceManagerRHI();
 
