Removed unused slate index buffer.

diff --git a/Engine/Source/Runtime/SlateRHIRenderer/Private/SlateElementIndexBuffer.cpp b/Engine/Source/Runtime/SlateRHIRenderer/Private/SlateElementIndexBuffer.cpp
deleted file mode 100644
diff --git a/Engine/Source/Runtime/SlateRHIRenderer/Private/SlateElementIndexBuffer.h b/Engine/Source/Runtime/SlateRHIRenderer/Private/SlateElementIndexBuffer.h
deleted file mode 100644
diff --git a/Engine/Source/Runtime/SlateRHIRenderer/Private/SlateRHIRenderingPolicy.h b/Engine/Source/Runtime/SlateRHIRenderer/Private/SlateRHIRenderingPolicy.h
--- a/Engine/Source/Runtime/SlateRHIRenderer/Private/SlateRHIRenderingPolicy.h
+++ b/Engine/Source/Runtime/SlateRHIRenderer/Private/SlateRHIRenderingPolicy.h
@@ -12,7 +12,6 @@
 #include "Rendering/DrawElements.h"
 #include "Rendering/RenderingPolicy.h"
 #include "Layout/Clipping.h"
-#include "SlateElementIndexBuffer.h"
 #include "SlateElementVertexBuffer.h"
 #include "SlateRHIResourceManager.h"
 #include "Shader.h"
diff --git a/Engine/Source/Runtime/SlateRHIRenderer/Private/SlateRHIResourceManager.h b/Engine/Source/Runtime/SlateRHIRenderer/Private/SlateRHIResourceManager.h
--- a/Engine/Source/Runtime/SlateRHIRenderer/Private/SlateRHIResourceManager.h
+++ b/Engine/Source/Runtime/SlateRHIRenderer/Private/SlateRHIResourceManager.h
@@ -13,7 +13,6 @@
 #include "Rendering/DrawElements.h"
 #include "Materials/MaterialInterface.h"
 #include "Tickable.h"
-#include "SlateElementIndexBuffer.h"
 #include "SlateElementVertexBuffer.h"
 
 class FSlateAtlasedTextureResource;
@@ -109,16 +108,6 @@ private:
 	FObjectResourceMap ObjectMap;
 };
 
-
-struct FCachedRenderBuffers
-{
-	TSlateElementVertexBuffer<FSlateVertex> VertexBuffer;
-	FSlateElementIndexBuffer IndexBuffer;
-
-	FGraphEventRef ReleaseResourcesFence;
-};
-
-
 /**
  * Stores a mapping of texture names to their RHI texture resource               
  */
