Reworked RDG dispatch logic to cut a new parallel pass set when a dispatch is requested. Also reworked the Slate 8 swapchain limit workaround so we only dispatch after 8 instead of every time.

diff --git a/Engine/Source/Runtime/RenderCore/Private/RenderGraphBuilder.cpp b/Engine/Source/Runtime/RenderCore/Private/RenderGraphBuilder.cpp
--- a/Engine/Source/Runtime/RenderCore/Private/RenderGraphBuilder.cpp
+++ b/Engine/Source/Runtime/RenderCore/Private/RenderGraphBuilder.cpp
@@ -2833,7 +2833,6 @@ void FRDGBuilder::SetupParallelExecute(TStaticArray<void*, MAX_NUM_GPUS> const&
 		}
 
 		bTaskModeAsync &= bPassTaskModeAsync;
-		bDispatchAfterExecute |= Pass->bDispatchAfterExecute;
 
 		ParallelPassCandidates.Emplace(Pass);
 
@@ -2842,6 +2841,12 @@ void FRDGBuilder::SetupParallelExecute(TStaticArray<void*, MAX_NUM_GPUS> const&
 			ParallelPassCandidatesWorkload += Pass->Workload;
 		}
 
+		if (Pass->bDispatchAfterExecute)
+		{
+			bDispatchAfterExecute = true;
+			FlushParallelPassCandidates();
+		}
+
 		if (ParallelPassCandidatesWorkload >= (uint32)GRDGParallelExecutePassMax)
 		{
 			FlushParallelPassCandidates();
diff --git a/Engine/Source/Runtime/SlateRHIRenderer/Private/SlateRHIRenderer.cpp b/Engine/Source/Runtime/SlateRHIRenderer/Private/SlateRHIRenderer.cpp
--- a/Engine/Source/Runtime/SlateRHIRenderer/Private/SlateRHIRenderer.cpp
+++ b/Engine/Source/Runtime/SlateRHIRenderer/Private/SlateRHIRenderer.cpp
@@ -920,9 +920,6 @@ FSlateDrawWindowPassOutputs FSlateRHIRenderer::DrawWindow_RenderThread(FRDGBuild
 		}
 	}
 
-	// Dispatch work after each window since certain RHI's can't record commands for multiple swap chains.
-	GraphBuilder.AddDispatchHint();
-
 	FSlateDrawWindowPassOutputs Outputs;
 	Outputs.ViewportRHI = ViewportInfo.ViewportRHI;
 	Outputs.ViewportTextureRHI = ViewportTextureRHI;
@@ -1332,9 +1329,18 @@ void FSlateRHIRenderer::DrawWindows_Private(FSlateDrawBuffer& WindowDrawBuffer)
 					DeferredUpdateContext.Renderer->DrawWindowToTarget_RenderThread(GraphBuilder, DeferredUpdateContext);
 				}
 
+				int32 NumSwapChains = 0;
+
 				for (const FSlateDrawWindowPassInputs& DrawWindowPassInputs : DrawWindowsCommand->Windows)
 				{
 					DrawWindowPassOutputs.Emplace(DrawWindowPassInputs.Renderer->DrawWindow_RenderThread(GraphBuilder, DrawWindowPassInputs));
+
+					// D3D12 can't handle more than 8 swap chains. Force a command list flush to avoid this limitation until it can be handled at the RHI level.
+					if (++NumSwapChains == 8)
+					{
+						GraphBuilder.AddDispatchHint();
+						NumSwapChains = 0;
+					}
 				}
 
 				GraphBuilder.Execute();
