Fixed PostProcessBlur sampling from the wrong viewport texture in the editor.

diff --git a/Engine/Source/Runtime/Engine/Classes/Engine/Engine.h b/Engine/Source/Runtime/Engine/Classes/Engine/Engine.h
--- a/Engine/Source/Runtime/Engine/Classes/Engine/Engine.h
+++ b/Engine/Source/Runtime/Engine/Classes/Engine/Engine.h
@@ -2819,7 +2819,7 @@ public:
 	ENGINE_API virtual bool HasMultipleLocalPlayers(UWorld* InWorld);
 
 	/** @return whether we're currently running with stereoscopic 3D enabled for the specified viewport (or globally, if viewport is nullptr) */
-	ENGINE_API bool IsStereoscopic3D(FViewport* InViewport = nullptr);
+	ENGINE_API bool IsStereoscopic3D(const FViewport* InViewport = nullptr) const;
 
 	/**
 	 * Adds a world location as a secondary view location for purposes of texture streaming.
diff --git a/Engine/Source/Runtime/Engine/Private/Slate/SceneViewport.cpp b/Engine/Source/Runtime/Engine/Private/Slate/SceneViewport.cpp
--- a/Engine/Source/Runtime/Engine/Private/Slate/SceneViewport.cpp
+++ b/Engine/Source/Runtime/Engine/Private/Slate/SceneViewport.cpp
@@ -1317,6 +1317,11 @@ FSlateShaderResource* FSceneViewport::GetViewportRenderTargetTexture() const
 	return (BufferedSlateHandles.Num() != 0) ? BufferedSlateHandles[CurrentBufferedTargetIndex] : nullptr;
 }
 
+bool FSceneViewport::IsStereoscopic3D() const
+{
+	return GEngine->IsStereoscopic3D(this);
+}
+
 void FSceneViewport::SetDebugCanvas(TSharedPtr<SDebugCanvas> InDebugCanvas)
 {
 	DebugCanvas = InDebugCanvas;
@@ -1997,7 +2002,7 @@ void FSceneViewport::WindowRenderTargetUpdate(FSlateRenderer* Renderer, SWindow*
 	if (Renderer && Window)
 	{
 		// The viewport provider interface is ONLY used for stereo VR compositing to blit the intermediate viewport target over to the VR swap chain.
-		Renderer->SetWindowRenderTarget(*Window, UseSeparateRenderTarget() && GEngine->IsStereoscopic3D(this) ? this : nullptr);
+		Renderer->SetWindowRenderTarget(*Window, UseSeparateRenderTarget() ? this : nullptr);
 	}
 }
 
diff --git a/Engine/Source/Runtime/Engine/Private/UnrealEngine.cpp b/Engine/Source/Runtime/Engine/Private/UnrealEngine.cpp
--- a/Engine/Source/Runtime/Engine/Private/UnrealEngine.cpp
+++ b/Engine/Source/Runtime/Engine/Private/UnrealEngine.cpp
@@ -4317,7 +4317,7 @@ bool UEngine::HasMultipleLocalPlayers(UWorld* InWorld)
 }
 
 /** @return whether we're currently running with stereoscopic 3D enabled */
-bool UEngine::IsStereoscopic3D(FViewport* InViewport)
+bool UEngine::IsStereoscopic3D(const FViewport* InViewport) const
 {
 	return (!InViewport || InViewport->IsStereoRenderingAllowed()) &&
 		(StereoRenderingDevice.IsValid() && StereoRenderingDevice->IsStereoEnabled());
diff --git a/Engine/Source/Runtime/Engine/Public/Slate/SceneViewport.h b/Engine/Source/Runtime/Engine/Public/Slate/SceneViewport.h
--- a/Engine/Source/Runtime/Engine/Public/Slate/SceneViewport.h
+++ b/Engine/Source/Runtime/Engine/Public/Slate/SceneViewport.h
@@ -209,6 +209,7 @@ public:
 
 	/** ISlateViewport interface */
 	ENGINE_API virtual FSlateShaderResource* GetViewportRenderTargetTexture() const override;
+	ENGINE_API virtual bool IsStereoscopic3D() const override;
 	ENGINE_API virtual void OnDrawViewport( const FGeometry& AllottedGeometry, const FSlateRect& MyCullingRect, FSlateWindowElementList& OutDrawElements, int32 LayerId, const FWidgetStyle& InWidgetStyle, bool bParentEnabled ) override;
 	ENGINE_API virtual FCursorReply OnCursorQuery( const FGeometry& MyGeometry, const FPointerEvent& CursorEvent ) override;
 	ENGINE_API virtual TOptional<TSharedRef<SWidget>> OnMapCursor(const FCursorReply& CursorReply) override;
diff --git a/Engine/Source/Runtime/SlateCore/Public/Textures/SlateShaderResource.h b/Engine/Source/Runtime/SlateCore/Public/Textures/SlateShaderResource.h
--- a/Engine/Source/Runtime/SlateCore/Public/Textures/SlateShaderResource.h
+++ b/Engine/Source/Runtime/SlateCore/Public/Textures/SlateShaderResource.h
@@ -212,4 +212,6 @@ class IViewportRenderTargetProvider
 {
 public:
 	virtual FSlateShaderResource* GetViewportRenderTargetTexture() = 0;
+
+	virtual bool IsStereoscopic3D() const = 0;
 };
diff --git a/Engine/Source/Runtime/SlateRHIRenderer/Private/SlateRHIRenderer.cpp b/Engine/Source/Runtime/SlateRHIRenderer/Private/SlateRHIRenderer.cpp
--- a/Engine/Source/Runtime/SlateRHIRenderer/Private/SlateRHIRenderer.cpp
+++ b/Engine/Source/Runtime/SlateRHIRenderer/Private/SlateRHIRenderer.cpp
@@ -174,6 +174,15 @@ struct FSlateViewportInfo : public FRenderResource
 		}
 		return nullptr;
 	}
+
+	bool IsStereoscopic3D() const
+	{
+		if (RenderTargetProvider)
+		{
+			return RenderTargetProvider->IsStereoscopic3D();
+		}
+		return false;
+	}
 };
 
 struct FSlatePostProcessUpdateRequest
@@ -686,7 +695,7 @@ FSlateDrawWindowPassOutputs FSlateRHIRenderer::DrawWindow_RenderThread(FRDGBuild
 		FRHITexture* SwapChainTextureRHI = RHIGetViewportBackBuffer(ViewportInfo.ViewportRHI);
 
 		// Only render to the intermediate viewport render target if stereo rendering is enabled, which we'll then composite later.
-		const bool bCompositeStereoToSwapChain = ViewportTextureRHI && GEngine && GEngine->StereoRenderingDevice.IsValid();
+		const bool bCompositeStereoToSwapChain = ViewportTextureRHI && GEngine && GEngine->StereoRenderingDevice.IsValid() && ViewportInfo.IsStereoscopic3D();
 
 		// The output texture is what we ultimately render or composite slate elements into.
 		OutputTextureRHI = bCompositeStereoToSwapChain ? ViewportTextureRHI : SwapChainTextureRHI;
