Make USlateRHIRendererSettings::Get{Mutable} only return if the object is already loaded. Otherwise it will create a blank object that is never overwritten.

diff --git a/Engine/Source/Runtime/SlateRHIRenderer/Private/SlateMaterialResource.cpp b/Engine/Source/Runtime/SlateRHIRenderer/Private/SlateMaterialResource.cpp
--- a/Engine/Source/Runtime/SlateRHIRenderer/Private/SlateMaterialResource.cpp
+++ b/Engine/Source/Runtime/SlateRHIRenderer/Private/SlateMaterialResource.cpp
@@ -97,16 +97,20 @@ FSlateMaterialResource::FSlateMaterialResource(const UMaterialInterface& InMater
 		MaterialObject->GetUsedTextures(OutUsedTextures, EMaterialQualityLevel::Num, true, ERHIFeatureLevel::Num, true);
 
 		CachedSlatePostBuffers = ESlatePostRT::None;
-		for (const UTexture* OutUsedTexture : OutUsedTextures)
+
+		if (const USlateRHIRendererSettings* RendererSettings = USlateRHIRendererSettings::Get())
 		{
-			for (const TPair<ESlatePostRT, FSlatePostSettings>& SlatePostSetting : USlateRHIRendererSettings::Get()->GetSlatePostSettings())
+			for (const UTexture* OutUsedTexture : OutUsedTextures)
 			{
-				const ESlatePostRT SlatePostBitflag = SlatePostSetting.Key;
-				const FSlatePostSettings& SlatePostSettingValue = SlatePostSetting.Value;
-
-				if (SlatePostSettingValue.bEnabled && OutUsedTexture && OutUsedTexture->GetPathName() == SlatePostSettingValue.GetPathToSlatePostRT())
+				for (const TPair<ESlatePostRT, FSlatePostSettings>& SlatePostSetting : RendererSettings->GetSlatePostSettings())
 				{
-					CachedSlatePostBuffers |= SlatePostBitflag;
+					const ESlatePostRT SlatePostBitflag = SlatePostSetting.Key;
+					const FSlatePostSettings& SlatePostSettingValue = SlatePostSetting.Value;
+
+					if (SlatePostSettingValue.bEnabled && OutUsedTexture && OutUsedTexture->GetPathName() == SlatePostSettingValue.GetPathToSlatePostRT())
+					{
+						CachedSlatePostBuffers |= SlatePostBitflag;
+					}
 				}
 			}
 		}
@@ -163,16 +167,19 @@ void FSlateMaterialResource::UpdateMaterial(const UMaterialInterface& InMaterial
 		MaterialObject->GetUsedTextures(OutUsedTextures, EMaterialQualityLevel::Num, true, ERHIFeatureLevel::Num, true);
 
 		CachedSlatePostBuffers = ESlatePostRT::None;
-		for (const UTexture* OutUsedTexture : OutUsedTextures)
+		if (const USlateRHIRendererSettings* RendererSettings = USlateRHIRendererSettings::Get())
 		{
-			for (const TPair<ESlatePostRT, FSlatePostSettings>& SlatePostSetting : USlateRHIRendererSettings::Get()->GetSlatePostSettings())
+			for (const UTexture* OutUsedTexture : OutUsedTextures)
 			{
-				const ESlatePostRT SlatePostBitflag = SlatePostSetting.Key;
-				const FSlatePostSettings& SlatePostSettingValue = SlatePostSetting.Value;
-
-				if (SlatePostSettingValue.bEnabled && OutUsedTexture && OutUsedTexture->GetPathName() == SlatePostSettingValue.GetPathToSlatePostRT())
+				for (const TPair<ESlatePostRT, FSlatePostSettings>& SlatePostSetting : RendererSettings->GetSlatePostSettings())
 				{
-					CachedSlatePostBuffers |= SlatePostBitflag;
+					const ESlatePostRT SlatePostBitflag = SlatePostSetting.Key;
+					const FSlatePostSettings& SlatePostSettingValue = SlatePostSetting.Value;
+
+					if (SlatePostSettingValue.bEnabled && OutUsedTexture && OutUsedTexture->GetPathName() == SlatePostSettingValue.GetPathToSlatePostRT())
+					{
+						CachedSlatePostBuffers |= SlatePostBitflag;
+					}
 				}
 			}
 		}
diff --git a/Engine/Source/Runtime/SlateRHIRenderer/Private/SlateRHIRenderer.cpp b/Engine/Source/Runtime/SlateRHIRenderer/Private/SlateRHIRenderer.cpp
--- a/Engine/Source/Runtime/SlateRHIRenderer/Private/SlateRHIRenderer.cpp
+++ b/Engine/Source/Runtime/SlateRHIRenderer/Private/SlateRHIRenderer.cpp
@@ -1075,7 +1075,7 @@ void FSlateRHIRenderer::DrawWindows_Private(FSlateDrawBuffer& WindowDrawBuffer)
 	const FGameTime AppDilatedTime = FGameTime::CreateDilated(FPlatformTime::Seconds() - GStartTime, AppDeltaTime, FApp::GetCurrentTime() - GStartTime, AppDeltaTime);
 	const FVector2f AppCursorPosition = FSlateApplication::Get().GetCursorPos();
 	const bool bAppCanRender = GIsClient && !IsRunningCommandlet() && !GUsingNullRHI;
-	const bool bAppCanRenderPostProcess = IsInGameThread() && UAssetManager::IsInitialized() && bAppCanRender && CVarCopyBackbufferToSlatePostRenderTargets.GetValueOnGameThread() > 0;
+	const bool bAppCanRenderPostProcess = RendererSettings && IsInGameThread() && UAssetManager::IsInitialized() && bAppCanRender && CVarCopyBackbufferToSlatePostRenderTargets.GetValueOnGameThread() > 0;
 	EPixelFormat AppViewportSceneFormat = PF_Unknown;
 	FIntPoint AppViewportExtentMax = FIntPoint::ZeroValue;
 	ESlatePostRT PostProcessAnyUsedBits = ESlatePostRT::None;
diff --git a/Engine/Source/Runtime/SlateRHIRenderer/Private/SlateUTextureResource.cpp b/Engine/Source/Runtime/SlateRHIRenderer/Private/SlateUTextureResource.cpp
--- a/Engine/Source/Runtime/SlateRHIRenderer/Private/SlateUTextureResource.cpp
+++ b/Engine/Source/Runtime/SlateRHIRenderer/Private/SlateUTextureResource.cpp
@@ -89,14 +89,17 @@ FSlateUTextureResource::FSlateUTextureResource(UTexture* InTexture)
 		Proxy->Resource = this;
 
 		CachedSlatePostBuffers = ESlatePostRT::None;
-		for (const TPair<ESlatePostRT, FSlatePostSettings>& SlatePostSetting : USlateRHIRendererSettings::Get()->GetSlatePostSettings())
+		if (const USlateRHIRendererSettings* RendererSettings = USlateRHIRendererSettings::Get())
 		{
-			const ESlatePostRT SlatePostBitflag = SlatePostSetting.Key;
-			const FSlatePostSettings& SlatePostSettingValue = SlatePostSetting.Value;
-
-			if (SlatePostSettingValue.bEnabled && InTexture && InTexture->GetPathName() == SlatePostSettingValue.GetPathToSlatePostRT())
+			for (const TPair<ESlatePostRT, FSlatePostSettings>& SlatePostSetting : RendererSettings->GetSlatePostSettings())
 			{
-				CachedSlatePostBuffers |= SlatePostBitflag;
+				const ESlatePostRT SlatePostBitflag = SlatePostSetting.Key;
+				const FSlatePostSettings& SlatePostSettingValue = SlatePostSetting.Value;
+
+				if (SlatePostSettingValue.bEnabled && InTexture && InTexture->GetPathName() == SlatePostSettingValue.GetPathToSlatePostRT())
+				{
+					CachedSlatePostBuffers |= SlatePostBitflag;
+				}
 			}
 		}
 	}
@@ -127,14 +130,17 @@ void FSlateUTextureResource::UpdateTexture(UTexture* InTexture)
 	if (Proxy && TextureObject)
 	{
 		CachedSlatePostBuffers = ESlatePostRT::None;
-		for (const TPair<ESlatePostRT, FSlatePostSettings>& SlatePostSetting : USlateRHIRendererSettings::Get()->GetSlatePostSettings())
+		if (const USlateRHIRendererSettings* RendererSettings = USlateRHIRendererSettings::Get())
 		{
-			const ESlatePostRT SlatePostBitflag = SlatePostSetting.Key;
-			const FSlatePostSettings& SlatePostSettingValue = SlatePostSetting.Value;
-
-			if (SlatePostSettingValue.bEnabled && InTexture && InTexture->GetPathName() == SlatePostSettingValue.GetPathToSlatePostRT())
+			for (const TPair<ESlatePostRT, FSlatePostSettings>& SlatePostSetting : RendererSettings->GetSlatePostSettings())
 			{
-				CachedSlatePostBuffers |= SlatePostBitflag;
+				const ESlatePostRT SlatePostBitflag = SlatePostSetting.Key;
+				const FSlatePostSettings& SlatePostSettingValue = SlatePostSetting.Value;
+
+				if (SlatePostSettingValue.bEnabled && InTexture && InTexture->GetPathName() == SlatePostSettingValue.GetPathToSlatePostRT())
+				{
+					CachedSlatePostBuffers |= SlatePostBitflag;
+				}
 			}
 		}
 
diff --git a/Engine/Source/Runtime/SlateRHIRenderer/Public/SlateRHIRendererSettings.h b/Engine/Source/Runtime/SlateRHIRenderer/Public/SlateRHIRendererSettings.h
--- a/Engine/Source/Runtime/SlateRHIRenderer/Public/SlateRHIRendererSettings.h
+++ b/Engine/Source/Runtime/SlateRHIRenderer/Public/SlateRHIRendererSettings.h
@@ -72,14 +72,21 @@ class SLATERHIRENDERER_API USlateRHIRendererSettings : public UDeveloperSettings
 	GENERATED_BODY()
 
 public:
-		
 	static const USlateRHIRendererSettings* Get()
-	{ 
+	{
+		if (!IsClassLoaded<USlateRHIRendererSettings>())
+		{
+			return nullptr;
+		}
 		return GetDefault<USlateRHIRendererSettings>();
 	}
 
 	static USlateRHIRendererSettings* GetMutable()
 	{
+		if (!IsClassLoaded<USlateRHIRendererSettings>())
+		{
+			return nullptr;
+		}
 		return GetMutableDefault<USlateRHIRendererSettings>();
 	}
 
diff --git a/Engine/Source/Runtime/UMG/Private/Slate/SPostBufferUpdate.cpp b/Engine/Source/Runtime/UMG/Private/Slate/SPostBufferUpdate.cpp
--- a/Engine/Source/Runtime/UMG/Private/Slate/SPostBufferUpdate.cpp
+++ b/Engine/Source/Runtime/UMG/Private/Slate/SPostBufferUpdate.cpp
@@ -59,9 +59,16 @@ public:
 
 void FPostBufferUpdater::Draw_RenderThread(FRDGBuilder& GraphBuilder, const FDrawPassInputs& Inputs)
 {
+	const USlateRHIRendererSettings* RendererSettings = USlateRHIRendererSettings::Get();
+
+	if (!RendererSettings)
+	{
+		return;
+	}
+
 	for (ESlatePostRT SlatePostBufferBit : MakeFlagsRange(Inputs.UsedSlatePostBuffers & BuffersToUpdate_Renderthread))
 	{
-		UTextureRenderTarget2D* SlatePostBuffer = Cast<UTextureRenderTarget2D>(USlateRHIRendererSettings::Get()->TryGetPostBufferRT(SlatePostBufferBit));
+		UTextureRenderTarget2D* SlatePostBuffer = Cast<UTextureRenderTarget2D>(RendererSettings->TryGetPostBufferRT(SlatePostBufferBit));
 		if (!SlatePostBuffer)
 		{
 			continue;
@@ -113,16 +120,19 @@ void FPostBufferUpdater::PostCustomElementAdded(FSlateElementBatcher& ElementBat
 	}
 
 	// Give proxies a chance to update their renderthread values.
-	for (ESlatePostRT SlatePostBufferBit : MakeFlagsRange(BuffersToUpdate_Renderthread))
+	if (const USlateRHIRendererSettings* RendererSettings = USlateRHIRendererSettings::Get())
 	{
-		if (!USlateRHIRendererSettings::Get()->GetSlatePostSetting(SlatePostBufferBit).bEnabled)
+		for (ESlatePostRT SlatePostBufferBit : MakeFlagsRange(BuffersToUpdate_Renderthread))
 		{
-			continue;
-		}
+			if (!RendererSettings->GetSlatePostSetting(SlatePostBufferBit).bEnabled)
+			{
+				continue;
+			}
 
-		if (TSharedPtr<FSlateRHIPostBufferProcessorProxy> PostProcessorProxy = USlateFXSubsystem::GetPostProcessorProxy(SlatePostBufferBit))
-		{
-			PostProcessorProxy->OnUpdateValuesRenderThread();
+			if (TSharedPtr<FSlateRHIPostBufferProcessorProxy> PostProcessorProxy = USlateFXSubsystem::GetPostProcessorProxy(SlatePostBufferBit))
+			{
+				PostProcessorProxy->OnUpdateValuesRenderThread();
+			}
 		}
 	}
 }
@@ -184,11 +194,14 @@ void SPostBufferUpdate::SetBuffersToUpdate(const TArrayView<ESlatePostRT> InBuff
 	if (PostBufferUpdater && !PostBufferUpdater->bBuffersToUpdateInitialized)
 	{
 		PostBufferUpdater->BuffersToUpdate_Renderthread = ESlatePostRT::None;
-		for (ESlatePostRT BufferToUpdate : BuffersToUpdate)
+		if (const USlateRHIRendererSettings* RendererSettings = USlateRHIRendererSettings::Get())
 		{
-			if (USlateRHIRendererSettings::Get()->GetSlatePostSetting(BufferToUpdate).bEnabled)
+			for (ESlatePostRT BufferToUpdate : BuffersToUpdate)
 			{
-				PostBufferUpdater->BuffersToUpdate_Renderthread |= BufferToUpdate;
+				if (RendererSettings->GetSlatePostSetting(BufferToUpdate).bEnabled)
+				{
+					PostBufferUpdater->BuffersToUpdate_Renderthread |= BufferToUpdate;
+				}
 			}
 		}
 
