Fixed Slate post process blur not respecting scissor clipping regions.

diff --git a/Engine/Source/Runtime/SlateRHIRenderer/Private/SlatePostProcessor.cpp b/Engine/Source/Runtime/SlateRHIRenderer/Private/SlatePostProcessor.cpp
--- a/Engine/Source/Runtime/SlateRHIRenderer/Private/SlatePostProcessor.cpp
+++ b/Engine/Source/Runtime/SlateRHIRenderer/Private/SlatePostProcessor.cpp
@@ -265,14 +265,22 @@ void AddSlatePostProcessUpsamplePass(FRDGBuilder& GraphBuilder, const FSlatePost
 		ERDGPassFlags::Raster,
 		[OutputViewport, InputViewport, ClippingElementsViewRect = Inputs.ClippingElementsViewRect, PipelineState, PixelShader, ClippingOp = Inputs.ClippingOp, PassParameters](FRDGAsyncTask, FRHICommandList& RHICmdList)
 	{
-		if (ClippingOp)
+		if (ClippingOp && ClippingOp->Method == EClippingMethod::Stencil)
 		{
+			// Stencil clipping quads have their own viewport.
 			RHICmdList.SetViewport(ClippingElementsViewRect.Min.X, ClippingElementsViewRect.Min.Y, 0.0f, ClippingElementsViewRect.Max.X, ClippingElementsViewRect.Max.Y, 1.0f);
+
 			// Stencil clipping will issue its own draw calls.
 			SetSlateClipping(RHICmdList, ClippingOp, ClippingElementsViewRect);
 		}
 
 		RHICmdList.SetViewport(OutputViewport.Rect.Min.X, OutputViewport.Rect.Min.Y, 0.0f, OutputViewport.Rect.Max.X, OutputViewport.Rect.Max.Y, 1.0f);
+
+		if (ClippingOp && ClippingOp->Method == EClippingMethod::Scissor)
+		{
+			SetSlateClipping(RHICmdList, ClippingOp, ClippingElementsViewRect);
+		}
+
 		SetScreenPassPipelineState(RHICmdList, PipelineState);
 		SetShaderParameters(RHICmdList, PixelShader, PixelShader.GetPixelShader(), *PassParameters);
 		DrawScreenPass_PostSetup(RHICmdList, FScreenPassViewInfo(), OutputViewport, InputViewport, PipelineState, EScreenPassDrawFlags::None);
