Initialize render thread proxy values immediately after creation. Fixes issue where certain post process buffers would not be updated.

diff --git a/Engine/Source/Runtime/SlateRHIRenderer/Private/FX/SlateFXSubsystem.cpp b/Engine/Source/Runtime/SlateRHIRenderer/Private/FX/SlateFXSubsystem.cpp
--- a/Engine/Source/Runtime/SlateRHIRenderer/Private/FX/SlateFXSubsystem.cpp
+++ b/Engine/Source/Runtime/SlateRHIRenderer/Private/FX/SlateFXSubsystem.cpp
@@ -114,8 +114,10 @@ void USlateFXSubsystem::OnPreWorldInitialization(UWorld* World, const UWorld::In
 			{
 				if (TObjectPtr<USlateRHIPostBufferProcessor> BufferProcessor = NewObject<USlateRHIPostBufferProcessor>(this, PostSetting.PostProcessorClass))
 				{
+					TSharedPtr<FSlateRHIPostBufferProcessorProxy> PostProcessorProxy = BufferProcessor->GetRenderThreadProxy();
 					SlatePostBufferProcessors.Add(SlatePostBufferBit, BufferProcessor);
-					SlatePostBufferProcessorProxies.Add(SlatePostBufferBit, BufferProcessor->GetRenderThreadProxy());
+					SlatePostBufferProcessorProxies.Add(SlatePostBufferBit, PostProcessorProxy);
+					PostProcessorProxy->OnUpdateValuesRenderThread();
 				}
 			}
 		}
