PR #12463: Fix uninitialized NumInstances member variable in FSlateUpdatableInstanceBuffer

#rnx

diff --git a/Engine/Source/Runtime/SlateRHIRenderer/Private/SlateUpdatableBuffer.h b/Engine/Source/Runtime/SlateRHIRenderer/Private/SlateUpdatableBuffer.h
--- a/Engine/Source/Runtime/SlateRHIRenderer/Private/SlateUpdatableBuffer.h
+++ b/Engine/Source/Runtime/SlateRHIRenderer/Private/SlateUpdatableBuffer.h
@@ -44,5 +44,5 @@ private:
 
 
 private:
-	uint32 NumInstances;
+	uint32 NumInstances = 0;
 };
