Fixed Slate incorrect clipping batch due to not invalidating the last clipping op when a custom drawer is used.

diff --git a/Engine/Source/Runtime/SlateRHIRenderer/Private/SlateRHIRenderingPolicy.cpp b/Engine/Source/Runtime/SlateRHIRenderer/Private/SlateRHIRenderingPolicy.cpp
--- a/Engine/Source/Runtime/SlateRHIRenderer/Private/SlateRHIRenderingPolicy.cpp
+++ b/Engine/Source/Runtime/SlateRHIRenderer/Private/SlateRHIRenderingPolicy.cpp
@@ -1561,6 +1561,7 @@ void AddSlateDrawElementsPass(
 
 			// Reset cached clipping state since custom draws mutate render state.
 			LastClippingState = nullptr;
+			LastClippingOp = nullptr;
 			break;
 		}
 		case ESlateRenderBatchType::PostProcess:
