Only allow up to 8 windows to be batched together into one RDG builder instance at a time. This fixes intermittant GPU crashes when more windows are queued for present.

diff --git a/Engine/Source/Runtime/SlateRHIRenderer/Private/SlateRHIRenderer.cpp b/Engine/Source/Runtime/SlateRHIRenderer/Private/SlateRHIRenderer.cpp
--- a/Engine/Source/Runtime/SlateRHIRenderer/Private/SlateRHIRenderer.cpp
+++ b/Engine/Source/Runtime/SlateRHIRenderer/Private/SlateRHIRenderer.cpp
@@ -1077,6 +1077,70 @@ void FSlateRHIRenderer::PrepareToTakeHDRScreenshot(const FIntRect& Rect, TArray<
 	ScreenshotState.ColorDataHDR = OutColorData;
 }
 
+struct FSlateDrawWindowsCommand : public TConcurrentLinearObject<FSlateDrawWindowsCommand>
+{
+	bool IsEmpty() { return Windows.IsEmpty() && DeferredUpdates.IsEmpty(); }
+
+	TArray<FSlateDrawWindowPassInputs, FConcurrentLinearArrayAllocator> Windows;
+	TArray<FSlatePostProcessUpdateRequest, FConcurrentLinearArrayAllocator> PostProcessUpdates;
+	TArray<FRenderThreadUpdateContext, FConcurrentLinearArrayAllocator> DeferredUpdates;
+};
+
+void FSlateRHIRenderer::DrawWindows_RenderThread(FRHICommandListImmediate& RHICmdList, TConstArrayView<FSlateDrawWindowPassInputs> Windows, TConstArrayView<FRenderThreadUpdateContext> DeferredUpdates)
+{
+	struct FWindowPresentCommand
+	{
+		FWindowPresentCommand(const FSlateDrawWindowPassInputs& InInputs, const FSlateDrawWindowPassOutputs& InOutputs)
+			: Inputs(InInputs)
+			, Outputs(InOutputs)
+		{}
+
+		const FSlateDrawWindowPassInputs& Inputs;
+		FSlateDrawWindowPassOutputs Outputs;
+	};
+
+	TArray<FWindowPresentCommand, FConcurrentLinearArrayAllocator> WindowPresentCommands;
+	WindowPresentCommands.Reserve(Windows.Num());
+
+	int32 WindowIndex = 0;
+
+	do
+	{
+		{
+			FRDGBuilder GraphBuilder(RHICmdList, RDG_EVENT_NAME("Slate"), ERDGBuilderFlags::ParallelSetup | ERDGBuilderFlags::ParallelExecute);
+
+			for (const FRenderThreadUpdateContext& DeferredUpdateContext : DeferredUpdates)
+			{
+				DeferredUpdateContext.Renderer->DrawWindowToTarget_RenderThread(GraphBuilder, DeferredUpdateContext);
+			}
+
+			// D3D12 can't handle more than 8 swap chains at a time, start a new batch if we hit this amount and continue with a new builder.
+			for (int32 NumWindows = 0; NumWindows < 8 && WindowIndex < Windows.Num(); ++NumWindows, ++WindowIndex)
+			{
+				const FSlateDrawWindowPassInputs& DrawWindowPassInputs = Windows[WindowIndex];
+				WindowPresentCommands.Emplace(DrawWindowPassInputs, DrawWindow_RenderThread(GraphBuilder, DrawWindowPassInputs));
+			}
+
+			GraphBuilder.AddDispatchHint();
+			GraphBuilder.Execute();
+		}
+
+		for (const FRenderThreadUpdateContext& DeferredUpdateContext : DeferredUpdates)
+		{
+			DeferredUpdateContext.Renderer->ReleaseDrawBuffer(*DeferredUpdateContext.WindowDrawBuffer);
+		}
+
+		for (const FWindowPresentCommand& Command : WindowPresentCommands)
+		{
+			PresentWindow_RenderThread(RHICmdList, Command.Inputs, Command.Outputs);
+		}
+		WindowPresentCommands.Reset();
+
+		DeferredUpdates = {};
+
+	} while (WindowIndex < Windows.Num());
+}
+
 void FSlateRHIRenderer::DrawWindows_Private(FSlateDrawBuffer& WindowDrawBuffer)
 {
 	checkSlow(IsThreadSafeForSlateRendering());
@@ -1261,16 +1325,7 @@ void FSlateRHIRenderer::DrawWindows_Private(FSlateDrawBuffer& WindowDrawBuffer)
 		}
 	}
 
-	struct FDrawWindowsCommand : public TConcurrentLinearObject<FDrawWindowsCommand>
-	{
-		bool IsEmpty() { return Windows.IsEmpty() && DeferredUpdates.IsEmpty(); }
-
-		TArray<FSlateDrawWindowPassInputs, FConcurrentLinearArrayAllocator> Windows;
-		TArray<FSlatePostProcessUpdateRequest, FConcurrentLinearArrayAllocator> PostProcessUpdates;
-		TArray<FRenderThreadUpdateContext, FConcurrentLinearArrayAllocator> DeferredUpdates;
-	};
-
-	TUniquePtr<FDrawWindowsCommand> DrawWindowsCommand = MakeUnique<FDrawWindowsCommand>();
+	TUniquePtr<FSlateDrawWindowsCommand> DrawWindowsCommand = MakeUnique<FSlateDrawWindowsCommand>();
 	DrawWindowsCommand->Windows.Reserve(WindowsToRender.Num());
 	DrawWindowsCommand->DeferredUpdates = MoveTemp(DeferredUpdateContexts);
 	DrawWindowsCommand->PostProcessUpdates.Reserve(WindowsToRender.Num() * FMath::CountBits((int32)PostProcessAnyUsedBits));
@@ -1341,46 +1396,9 @@ void FSlateRHIRenderer::DrawWindows_Private(FSlateDrawBuffer& WindowDrawBuffer)
 
 	if (!DrawWindowsCommand->IsEmpty())
 	{
-		ENQUEUE_RENDER_COMMAND(SlateDrawWindowsCommand)([DrawWindowsCommand = MoveTemp(DrawWindowsCommand)](FRHICommandListImmediate& RHICmdList)
+		ENQUEUE_RENDER_COMMAND(SlateDrawWindowsCommand)([this, DrawWindowsCommand = MoveTemp(DrawWindowsCommand)](FRHICommandListImmediate& RHICmdList)
 		{
-			TArray<FSlateDrawWindowPassOutputs, FConcurrentLinearArrayAllocator> DrawWindowPassOutputs;
-			DrawWindowPassOutputs.Reserve(DrawWindowsCommand->Windows.Num());
-
-			{
-				FRDGBuilder GraphBuilder(RHICmdList, RDG_EVENT_NAME("Slate"), ERDGBuilderFlags::ParallelSetup | ERDGBuilderFlags::ParallelExecute);
-
-				for (const FRenderThreadUpdateContext& DeferredUpdateContext : DrawWindowsCommand->DeferredUpdates)
-				{
-					DeferredUpdateContext.Renderer->DrawWindowToTarget_RenderThread(GraphBuilder, DeferredUpdateContext);
-				}
-
-				int32 NumSwapChains = 0;
-
-				for (const FSlateDrawWindowPassInputs& DrawWindowPassInputs : DrawWindowsCommand->Windows)
-				{
-					DrawWindowPassOutputs.Emplace(DrawWindowPassInputs.Renderer->DrawWindow_RenderThread(GraphBuilder, DrawWindowPassInputs));
-
-					// D3D12 can't handle more than 8 swap chains. Force a command list flush to avoid this limitation until it can be handled at the RHI level.
-					if (++NumSwapChains == 8)
-					{
-						GraphBuilder.AddDispatchHint();
-						NumSwapChains = 0;
-					}
-				}
-
-				GraphBuilder.Execute();
-
-				for (const FRenderThreadUpdateContext& DeferredUpdateContext : DrawWindowsCommand->DeferredUpdates)
-				{
-					DeferredUpdateContext.Renderer->ReleaseDrawBuffer(*DeferredUpdateContext.WindowDrawBuffer);
-				}
-			}
-
-			for (int32 WindowIndex = 0; WindowIndex < DrawWindowsCommand->Windows.Num(); ++WindowIndex)
-			{
-				const FSlateDrawWindowPassInputs& DrawWindowPassInputs = DrawWindowsCommand->Windows[WindowIndex];
-				DrawWindowPassInputs.Renderer->PresentWindow_RenderThread(RHICmdList, DrawWindowPassInputs, DrawWindowPassOutputs[WindowIndex]);
-			}
+			DrawWindows_RenderThread(RHICmdList, DrawWindowsCommand->Windows, DrawWindowsCommand->DeferredUpdates);
 		});
 
 		check(DeferredUpdateContexts.IsEmpty());
diff --git a/Engine/Source/Runtime/SlateRHIRenderer/Private/SlateRHIRenderer.h b/Engine/Source/Runtime/SlateRHIRenderer/Private/SlateRHIRenderer.h
--- a/Engine/Source/Runtime/SlateRHIRenderer/Private/SlateRHIRenderer.h
+++ b/Engine/Source/Runtime/SlateRHIRenderer/Private/SlateRHIRenderer.h
@@ -24,6 +24,7 @@ class USlateRHIRendererSettings;
 class ISlateStyle;
 class SWindow;
 struct FSlatePostSettings;
+struct FSlateDrawWindowsCommand;
 
 extern FMatrix CreateSlateProjectionMatrix(uint32 Width, uint32 Height);
 
@@ -112,6 +113,8 @@ private:
 
 	void PresentWindow_RenderThread(FRHICommandListImmediate& RHICmdList, const FSlateDrawWindowPassInputs& DrawPassInputs, const FSlateDrawWindowPassOutputs& DrawPassOutputs);
 
+	void DrawWindows_RenderThread(FRHICommandListImmediate& RHICmdList, TConstArrayView<FSlateDrawWindowPassInputs> Windows, TConstArrayView<FRenderThreadUpdateContext> DeferredUpdates);
+
 	/** Loads all known textures from Slate styles */
 	void LoadUsedTextures();
 
