SlateRHIRenderingPolicy nop comment only

diff --git a/Engine/Source/Runtime/SlateRHIRenderer/Private/SlateRHIRenderingPolicy.cpp b/Engine/Source/Runtime/SlateRHIRenderer/Private/SlateRHIRenderingPolicy.cpp
--- a/Engine/Source/Runtime/SlateRHIRenderer/Private/SlateRHIRenderingPolicy.cpp
+++ b/Engine/Source/Runtime/SlateRHIRenderer/Private/SlateRHIRenderingPolicy.cpp
@@ -1125,6 +1125,7 @@ FSlateRenderBatchOp* CreateSlateRenderBatchOp(
 			bool bIsVirtualTexture = false;
 
 			// check if texture is using BC4 compression and set shader to render grayscale
+			// @todo Oodle : it seems strange to special-case only BC4/TC_Alpha here ; prefer calling ShouldUseGreyScaleEditorVisualization()
 			bool bUseTextureGrayscale = false;
 
 			if (ShaderResource != nullptr && ResourceType == ESlateShaderResource::TextureObject)
