FSlateTextureData can be set from FImage
lets us write simpler and better code to deal with images

diff --git a/Engine/Source/Editor/UnrealEd/Private/AssetThumbnail.cpp b/Engine/Source/Editor/UnrealEd/Private/AssetThumbnail.cpp
--- a/Engine/Source/Editor/UnrealEd/Private/AssetThumbnail.cpp
+++ b/Engine/Source/Editor/UnrealEd/Private/AssetThumbnail.cpp
@@ -2026,10 +2026,12 @@ bool FAssetThumbnailPool::LoadThumbnail(TSharedRef<FThumbnailInfo> ThumbnailInfo
 
 	if (FoundThumbnail)
 	{
-		if (FoundThumbnail->GetImageWidth() > 0 && FoundThumbnail->GetImageHeight() > 0 && FoundThumbnail->GetUncompressedImageData().Num() > 0)
+		FImageView Image = FoundThumbnail->GetImage();
+
+		if ( Image.GetNumPixels() > 0 )
 		{
 			// Make bulk data for updating the texture memory later
-			FSlateTextureData* BulkData = new FSlateTextureData(FoundThumbnail->GetImageWidth(), FoundThumbnail->GetImageHeight(), GPixelFormats[PF_B8G8R8A8].BlockBytes, FoundThumbnail->AccessImageData());
+			FSlateTextureData* BulkData = new FSlateTextureData(Image);
 
 			// Update the texture RHI
 			FThumbnailInfo_RenderThread ThumbInfo = ThumbnailInfo.Get();
diff --git a/Engine/Source/Editor/WorldBrowser/Private/Tiles/WorldTileThumbnails.cpp b/Engine/Source/Editor/WorldBrowser/Private/Tiles/WorldTileThumbnails.cpp
--- a/Engine/Source/Editor/WorldBrowser/Private/Tiles/WorldTileThumbnails.cpp
+++ b/Engine/Source/Editor/WorldBrowser/Private/Tiles/WorldTileThumbnails.cpp
@@ -95,10 +95,10 @@ FSlateTextureDataPtr FTileThumbnail::ToSlateTextureData(const FObjectThumbnail*
 		FIntPoint TargetSize = ThumbnailRenderTarget->GetSizeXY();
 		if (TargetSize == ImageSize)
 		{
-			const TArray<uint8>& ImageData = ObjectThumbnail->GetUncompressedImageData();
-			if (ImageData.Num())
+			FImageView Image = ObjectThumbnail->GetImage();
+			if ( Image.GetNumPixels() > 0 )
 			{
-				Result = MakeShareable(new FSlateTextureData(ImageData.GetData(), ImageSize.X, ImageSize.Y, 4));
+				Result = MakeShareable(new FSlateTextureData(Image));
 			}
 		}
 	}
diff --git a/Engine/Source/Editor/WorldBrowser/WorldBrowser.Build.cs b/Engine/Source/Editor/WorldBrowser/WorldBrowser.Build.cs
--- a/Engine/Source/Editor/WorldBrowser/WorldBrowser.Build.cs
+++ b/Engine/Source/Editor/WorldBrowser/WorldBrowser.Build.cs
@@ -23,8 +23,9 @@ public class WorldBrowser : ModuleRules
                 "CoreUObject", 
                 "EditorConfig",
                 "RenderCore",
+				"ImageCore",
                 "InputCore",
-                "Engine",
+				"Engine",
 				"Landscape",
                 "Slate",
 				"SlateCore",
diff --git a/Engine/Source/Runtime/Engine/Private/Slate/SlateTextures.cpp b/Engine/Source/Runtime/Engine/Private/Slate/SlateTextures.cpp
--- a/Engine/Source/Runtime/Engine/Private/Slate/SlateTextures.cpp
+++ b/Engine/Source/Runtime/Engine/Private/Slate/SlateTextures.cpp
@@ -129,19 +129,24 @@ void FSlateTexture2DRHIRef::SetRHIRef( FTextureRHIRef InRHIRef, uint32 InWidth,
 void FSlateTexture2DRHIRef::SetTextureData( FSlateTextureDataPtr NewTextureData )
 {
 	check( IsInRenderingThread() );
+
 	Width = NewTextureData->GetWidth();
 	Height = NewTextureData->GetHeight();
 	TextureData = NewTextureData;
+
+	// we will memcpy the TextureData into the RHI texture without checks later, so they must match exactly
+	//	typically BytesPerPixel == 4 and PixelFormat == BGRA8 , but other things are possible here for viewports & movies and such.
+	check( NewTextureData->GetBytesPerPixel() == GPixelFormats[PixelFormat].BlockBytes );
 }
 
 void FSlateTexture2DRHIRef::SetTextureData( FSlateTextureDataPtr NewTextureData, EPixelFormat InPixelFormat, ETextureCreateFlags InTexCreateFlags )
 {
 	check( IsInRenderingThread() );
-
-	SetTextureData( NewTextureData );
-
+	
 	PixelFormat = InPixelFormat;
 	TexCreateFlags = InTexCreateFlags;
+
+	SetTextureData( NewTextureData );
 }
 
 
@@ -222,7 +227,9 @@ void FSlateTexture2DRHIRef::UpdateTextureThreadSafe(const TArray<uint8>& Bytes)
 	if(IsInGameThread())
 	{
 		// Make bulk data for updating the texture memory later
-		FSlateTextureData* BulkData = new FSlateTextureData( Bytes.Num(), 0, 1, Bytes );
+		uint32 BytesPerPixel = GPixelFormats[PixelFormat].BlockBytes;
+		check( Width*Height*BytesPerPixel == Bytes.Num() );
+		FSlateTextureData* BulkData = new FSlateTextureData( Width,Height,BytesPerPixel, Bytes );
 
 		// Update the texture RHI
 		FSlateTexture2DRHIRef* ThisTexture = this;
@@ -248,7 +255,14 @@ void FSlateTexture2DRHIRef::UpdateTextureThreadSafeRaw(uint32 InWidth, uint32 In
 void FSlateTexture2DRHIRef::UpdateTextureThreadSafeWithTextureData(FSlateTextureData* BulkData)
 {
 	check(IsInGameThread());
+	
 	// Update the texture RHI
+	// BEWARE this is enqueued, but is not holding a ref on "BulkData" ; should probably use a TSharedPtr<FSlateTextureData>
+
+	// we will memcpy the TextureData into the RHI texture without checks later, so they must match exactly
+	//	typically BytesPerPixel == 4 and PixelFormat == BGRA8 , but other things are possible here for viewports & movies and such.
+	check( BulkData->GetBytesPerPixel() == GPixelFormats[PixelFormat].BlockBytes );
+
 	FSlateTexture2DRHIRef* ThisTexture = this;
 	ENQUEUE_RENDER_COMMAND(FSlateTexture2DRHIRef_UpdateTextureThreadSafeWithTextureData)(
 		[ThisTexture, BulkData](FRHICommandListImmediate& RHICmdList)
diff --git a/Engine/Source/Runtime/SlateCore/Public/Textures/SlateTextureData.h b/Engine/Source/Runtime/SlateCore/Public/Textures/SlateTextureData.h
--- a/Engine/Source/Runtime/SlateCore/Public/Textures/SlateTextureData.h
+++ b/Engine/Source/Runtime/SlateCore/Public/Textures/SlateTextureData.h
@@ -5,11 +5,15 @@
 #include "CoreMinimal.h"
 #include "Stats/Stats.h"
 #include "SlateGlobals.h"
+#include "ImageCore.h"
 
 DECLARE_MEMORY_STAT_EXTERN(TEXT("Texture Data Memory (CPU)"), STAT_SlateTextureDataMemory, STATGROUP_SlateMemory, SLATECORE_API);
 
 /**
  * Holds texture data for upload to a rendering resource
+ * Makes a copy of the bytes passed to it and holds ownership of the image data array.
+ *
+ * Note that "BytesPerPixel" is a variable but in practice this must be BGRA8-SRGB , and BytesPerPixel == 4
  */
 struct FSlateTextureData
 {
@@ -22,6 +26,7 @@ struct FSlateTextureData
 		INC_MEMORY_STAT_BY( STAT_SlateTextureDataMemory, Bytes.GetAllocatedSize() );
 	}
 
+	// Move from TArray :
 	FSlateTextureData(uint32 InWidth, uint32 InHeight, uint32 InBytesPerPixel, TArray<uint8>&& InBytes)
 		: Bytes(InBytes)
 		, Width(InWidth)
@@ -60,6 +65,24 @@ struct FSlateTextureData
 	{
 		INC_MEMORY_STAT_BY( STAT_SlateTextureDataMemory, Bytes.GetAllocatedSize() );
 	}
+	
+	FSlateTextureData( const FImageView &Other )
+		: Bytes()
+		, Width(0)
+		, Height(0)
+		, BytesPerPixel(0)
+	{
+		SetImage(Other);
+	}
+	
+	FSlateTextureData( FImage && Other )
+		: Bytes()
+		, Width(0)
+		, Height(0)
+		, BytesPerPixel(0)
+	{
+		SetImage(Other);
+	}
 
 	FSlateTextureData& operator=( const FSlateTextureData& Other )
 	{
@@ -87,6 +110,38 @@ struct FSlateTextureData
 		INC_MEMORY_STAT_BY( STAT_SlateTextureDataMemory, Bytes.GetAllocatedSize() );
 	}
 	
+	// copy FImageView into the FSlateTextureData
+	//	convert to BGRA8-SRGB if not already in that format
+	void SetImage( const FImageView & Image )
+	{
+		// if Image is already BGRA8-SRGB then this is just a memcpy
+		//	which is what we need anyway to copy the bytes into a new array
+		//	so there is no inefficiency in just always using the FImage copy here
+		FImage Converted;
+		Image.CopyTo(Converted,ERawImageFormat::BGRA8,EGammaSpace::sRGB);
+		SetImage( MoveTemp(Converted) );
+	}
+		
+	// move FImage into the FSlateTextureData
+	//	convert to BGRA8-SRGB if not already in that format
+	void SetImage( FImage && MoveFrom )
+	{
+		// change format if needed, nop if not
+		// MoveFrom is discardable so it's okay if we just change it in place
+		MoveFrom.ChangeFormat(ERawImageFormat::BGRA8,EGammaSpace::sRGB);
+
+		DEC_MEMORY_STAT_BY( STAT_SlateTextureDataMemory, Bytes.GetAllocatedSize() );
+
+		Width = MoveFrom.SizeX;
+		Height = MoveFrom.SizeY;
+		BytesPerPixel = 4;		
+		Bytes = MoveTemp(MoveFrom.RawData);
+
+		INC_MEMORY_STAT_BY( STAT_SlateTextureDataMemory, Bytes.GetAllocatedSize() );
+
+		MoveFrom.Reset();
+	}
+
 	void Empty()
 	{
 		DEC_MEMORY_STAT_BY( STAT_SlateTextureDataMemory, Bytes.GetAllocatedSize() );
@@ -121,7 +176,7 @@ struct FSlateTextureData
 
 private:
 
-	/** Raw uncompressed texture data */
+	/** Raw uncompressed texture data (in practice always FColor, BGRA8-sRGB */
 	TArray<uint8> Bytes;
 
 	/** Width of the texture */
@@ -130,7 +185,7 @@ private:
 	/** Height of the texture */
 	uint32 Height;
 
-	/** The number of bytes of each pixel */
+	/** The number of bytes of each pixel (in practice must always be 4) */
 	uint32 BytesPerPixel;
 };
 
diff --git a/Engine/Source/Runtime/SlateRHIRenderer/Private/SlateRHIResourceManager.cpp b/Engine/Source/Runtime/SlateRHIRenderer/Private/SlateRHIResourceManager.cpp
--- a/Engine/Source/Runtime/SlateRHIRenderer/Private/SlateRHIResourceManager.cpp
+++ b/Engine/Source/Runtime/SlateRHIRenderer/Private/SlateRHIResourceManager.cpp
@@ -17,8 +17,6 @@
 #include "Engine/Engine.h"
 #include "Slate/SlateTextures.h"
 #include "SlateRHITextureAtlas.h"
-#include "IImageWrapper.h"
-#include "IImageWrapperModule.h"
 #include "SlateNativeTextureResource.h"
 #include "SlateUTextureResource.h"
 #include "SlateMaterialResource.h"
@@ -441,17 +439,21 @@ void FSlateRHIResourceManager::CreateTextures( const TArray< const FSlateBrush*
 		}
 	}
 
-	// This must be loaded before going multi-threaded
-	FModuleManager::LoadModuleChecked<IImageWrapperModule>(FName("ImageWrapper"));
+	// ImageWrapper module must be loaded on game thread before going multi-threaded
+	//	-> note this is the first place that ImageWrapper module is loaded in LaunchEngineLoop
+	FModuleManager::Get().LoadModuleChecked(FName("ImageWrapper"));
 
 	TArray<FName> TexturesToLoad;
 	TextureInfoMap.GenerateKeyArray(TexturesToLoad);
 
-	const uint32 Stride = GPixelFormats[PF_R8G8B8A8].BlockBytes;
+	const uint32 BytesPerPixel = GPixelFormats[PF_R8G8B8A8].BlockBytes;
+	check( BytesPerPixel == 4 );
+
 	ParallelFor(
 		TexturesToLoad.Num(),
 		[&](int32 TextureIndex)
 		{
+			// TextureInfoMap is being accessed multi-threaded, it must not be changed during this phase
 			FName& TextureName = TexturesToLoad[TextureIndex];
 			FLoadTextureInfo& Info = TextureInfoMap.FindChecked(TextureName);
 
@@ -460,10 +462,13 @@ void FSlateRHIResourceManager::CreateTextures( const TArray< const FSlateBrush*
 			TArray<uint8> RawData;
 			Info.bSucceeded = LoadTexture(*Info.Brush, Width, Height, RawData);
 
-			Info.TextureData = MakeShareable(new FSlateTextureData(Width, Height, Stride, MoveTemp(RawData)));
+			if ( Info.bSucceeded )
+			{
+				Info.TextureData = MakeShareable(new FSlateTextureData(Width, Height, BytesPerPixel, MoveTemp(RawData)));
 
-			const bool bTooLargeForAtlas = (Width >= (uint32)MaxAltasedTextureSize.X || Height >= (uint32)MaxAltasedTextureSize.Y || Width >= AtlasSize || Height >= AtlasSize);
-			Info.bShouldAtlas &= !bTooLargeForAtlas;
+				const bool bTooLargeForAtlas = (Width >= (uint32)MaxAltasedTextureSize.X || Height >= (uint32)MaxAltasedTextureSize.Y || Width >= AtlasSize || Height >= AtlasSize);
+				Info.bShouldAtlas &= !bTooLargeForAtlas;
+			}
 		},
 		EParallelForFlags::Unbalanced
 	);
@@ -506,49 +511,25 @@ bool FSlateRHIResourceManager::LoadTexture( const FSlateBrush& InBrush, uint32&
  *
  * @param TextureName	The name of the texture to load
  */
-bool FSlateRHIResourceManager::LoadTexture( const FName& TextureName, const FString& ResourcePath, uint32& Width, uint32& Height, TArray<uint8>& DecodedImage )
+bool FSlateRHIResourceManager::LoadTexture( const FName& TextureName, const FString& ResourcePath, 
+	uint32& OutWidth, uint32& OutHeight, TArray<uint8>& OutDecodedImage )
 {
-	bool bSucceeded = false;
-	uint32 BytesPerPixel = 4;
+	OutWidth = OutHeight = 0;
 
-	TArray<uint8> RawFileData;
-	if( FFileHelper::LoadFileToArray( RawFileData, *ResourcePath ) )
+	FImage Image;
+	if ( ! FImageUtils::LoadImage(*ResourcePath,Image) )
 	{
-		IImageWrapperModule& ImageWrapperModule = FModuleManager::LoadModuleChecked<IImageWrapperModule>( FName("ImageWrapper") );
+		UE_LOG(LogSlate, Log, TEXT("Could not load file for Slate resource: [%s] file: [%s]"), *TextureName.ToString(),*ResourcePath );
+		return false;
+	}
 
-		//Try and determine format, if that fails assume PNG
-		EImageFormat ImageFormat = ImageWrapperModule.DetectImageFormat(RawFileData.GetData(), RawFileData.Num());
-		if (ImageFormat == EImageFormat::Invalid)
-		{
-			ImageFormat = EImageFormat::PNG;
-		}
-		TSharedPtr<IImageWrapper> ImageWrapper = ImageWrapperModule.CreateImageWrapper(ImageFormat, *ResourcePath);
+	Image.ChangeFormat(ERawImageFormat::BGRA8,EGammaSpace::sRGB);
 
-		if ( ImageWrapper.IsValid() && ImageWrapper->SetCompressed( RawFileData.GetData(), RawFileData.Num()) )
-		{
-			Width = ImageWrapper->GetWidth();
-			Height = ImageWrapper->GetHeight();
-			
-			if (ImageWrapper->GetRaw( ERGBFormat::BGRA, 8, DecodedImage))
-			{
-				bSucceeded = true;
-			}
-			else
-			{
-				UE_LOG(LogSlate, Log, TEXT("Invalid texture format for Slate resource only RGBA and RGB pngs are supported: %s"), *TextureName.ToString() );
-			}
-		}
-		else
-		{
-			UE_LOG(LogSlate, Log, TEXT("Only pngs are supported in Slate"));
-		}
-	}
-	else
-	{
-		UE_LOG(LogSlate, Log, TEXT("Could not find file for Slate resource: %s"), *TextureName.ToString() );
-	}
+	OutWidth = Image.SizeX;
+	OutHeight = Image.SizeY;
+	OutDecodedImage = MoveTemp(Image.RawData);
 
-	return bSucceeded;
+	return true;
 }
 
 FSlateShaderResourceProxy* FSlateRHIResourceManager::GenerateTextureResource( const FNewTextureInfo& Info, const FName TextureName )
@@ -729,6 +710,8 @@ TSharedPtr<FSlateDynamicTextureResource> FSlateRHIResourceManager::MakeDynamicTe
 
 	TextureResource->Proxy->ActualSize = FIntPoint(TextureData->GetWidth(), TextureData->GetHeight());
 
+	// TextureData hard-coded assumption to be BGRA8-SRGB
+	check( TextureData->GetBytesPerPixel() == 4 );
 
 	// Init render thread data
 	FSlateDynamicTextureResource* InTextureResource = TextureResource.Get();
diff --git a/Engine/Source/Runtime/SlateRHIRenderer/SlateRHIRenderer.Build.cs b/Engine/Source/Runtime/SlateRHIRenderer/SlateRHIRenderer.Build.cs
--- a/Engine/Source/Runtime/SlateRHIRenderer/SlateRHIRenderer.Build.cs
+++ b/Engine/Source/Runtime/SlateRHIRenderer/SlateRHIRenderer.Build.cs
@@ -17,7 +17,7 @@ public class SlateRHIRenderer : ModuleRules
                 "RHI",
                 "RenderCore",
 				"Renderer",
-                "ImageWrapper",
+                "ImageCore",
 				"HeadMountedDisplay"
 			}
 		);
