Slate: enable viewport scale uniform for cooked builds

diff --git a/Engine/Source/Runtime/SlateRHIRenderer/Private/SlateRHIRenderer.cpp b/Engine/Source/Runtime/SlateRHIRenderer/Private/SlateRHIRenderer.cpp
--- a/Engine/Source/Runtime/SlateRHIRenderer/Private/SlateRHIRenderer.cpp
+++ b/Engine/Source/Runtime/SlateRHIRenderer/Private/SlateRHIRenderer.cpp
@@ -1247,7 +1247,7 @@ void FSlateRHIRenderer::DrawWindows_Private(FSlateDrawBuffer& WindowDrawBuffer)
 			AppViewportSceneFormat = Viewport->GetSceneTargetFormat();
 			AppViewportExtentMax = AppViewportExtentMax.ComponentMax(ViewportExtent);
 
-			if (GIsEditor && ViewportScaleUI < 0.0f)
+			if (ViewportScaleUI < 0.0f)
 			{
 				ViewportScaleUI = GetDefault<UUserInterfaceSettings>()->GetDPIScaleBasedOnSize(ViewportExtent);
 			}
