Add PSO precaching for slate global and material PSOs

diff --git a/Engine/Source/Runtime/Engine/Private/Materials/Material.cpp b/Engine/Source/Runtime/Engine/Private/Materials/Material.cpp
--- a/Engine/Source/Runtime/Engine/Private/Materials/Material.cpp
+++ b/Engine/Source/Runtime/Engine/Private/Materials/Material.cpp
@@ -4646,7 +4646,7 @@ void UMaterial::PostLoad()
 			PreloadMaterialShaderMap(GetMaterialResource(GMaxRHIFeatureLevel), Unused);
 		}
 	}
-	else if (IsDeferredDecal())
+	else if (IsDeferredDecal() || IsUIMaterial() || (MaterialDomain == MD_LightFunction) || IsPostProcessMaterial())
 	{
 		FPSOPrecacheParams PSOPrecacheParams;
 		UMaterialInterface::PrecachePSOs(&FLocalVertexFactory::StaticType, PSOPrecacheParams);
diff --git a/Engine/Source/Runtime/Engine/Private/Materials/MaterialInstance.cpp b/Engine/Source/Runtime/Engine/Private/Materials/MaterialInstance.cpp
--- a/Engine/Source/Runtime/Engine/Private/Materials/MaterialInstance.cpp
+++ b/Engine/Source/Runtime/Engine/Private/Materials/MaterialInstance.cpp
@@ -3489,7 +3489,7 @@ void UMaterialInstance::PostLoad()
 			PreloadMaterialShaderMap(GetMaterialResource(GMaxRHIFeatureLevel), Unused);
 		}
 	}
-	else if (IsDeferredDecal())
+	else if (IsDeferredDecal() || IsUIMaterial() || IsPostProcessMaterial())
 	{
 		FPSOPrecacheParams PSOPrecacheParams;
 		UMaterialInterface::PrecachePSOs(&FLocalVertexFactory::StaticType, PSOPrecacheParams);
diff --git a/Engine/Source/Runtime/Renderer/Private/PostProcess/PostProcessMaterial.cpp b/Engine/Source/Runtime/Renderer/Private/PostProcess/PostProcessMaterial.cpp
--- a/Engine/Source/Runtime/Renderer/Private/PostProcess/PostProcessMaterial.cpp
+++ b/Engine/Source/Runtime/Renderer/Private/PostProcess/PostProcessMaterial.cpp
@@ -56,6 +56,13 @@ TAutoConsoleVariable<int32> CVarPostProcessingDisableMaterials(
 	TEXT(" Allows to disable post process materials. \n"),
 	ECVF_Scalability | ECVF_RenderThreadSafe);
 
+static int32 GPostProcessingMaterialPSOPrecache = 1;
+static FAutoConsoleVariableRef CVarGPostProcessingMaterialPSOPrecache(
+	TEXT("r.PSOPrecache.PostProcessingMaterial"), 
+	GPostProcessingMaterialPSOPrecache, 
+	TEXT("Precache all possible required PSOs for loaded PostProcessing Materials."),
+	ECVF_ReadOnly);
+
 static FName NAME_SceneColor("SceneColor");
 
 static bool IsPostProcessStencilTestAllowed()
@@ -122,7 +129,7 @@ static bool IsMaterialBlendEnabled(const FMaterial* Material)
 {
 	check(Material);
 
-	return Material->GetBlendableOutputAlpha() && CVarPostProcessAllowBlendModes.GetValueOnRenderThread() != 0;
+	return Material->GetBlendableOutputAlpha() && CVarPostProcessAllowBlendModes.GetValueOnAnyThread() != 0;
 }
 
 static FRHIBlendState* GetMaterialBlendState(const FMaterial* Material)
@@ -346,6 +353,34 @@ public:
 	}
 };
 
+static bool GetMaterialShaders(
+	const FMaterial& Material, 
+	bool bManualStencilTest, 
+	bool bNeuralPostProcessPrepass,
+	TShaderRef<FPostProcessMaterialVS>& OutVertexShader,
+	TShaderRef<FPostProcessMaterialPS>& OutPixelShader)
+{
+	FMaterialShaderTypes ShaderTypes;
+
+	FPostProcessMaterialPS::FPermutationDomain PermutationVectorPS;
+	PermutationVectorPS.Set<FPostProcessMaterialPS::FManualStencilTestDim>(bManualStencilTest);
+	PermutationVectorPS.Set<FPostProcessMaterialPS::FNeuralPostProcessPrePass>(bNeuralPostProcessPrepass);
+
+	ShaderTypes.AddShaderType<FPostProcessMaterialVS>();
+	ShaderTypes.AddShaderType<FPostProcessMaterialPS>(PermutationVectorPS.ToDimensionValueId());
+
+	FMaterialShaders Shaders;
+	if (!Material.TryGetShaders(ShaderTypes, nullptr, Shaders))
+	{
+		return false;
+	}
+	
+	Shaders.TryGetVertexShader(OutVertexShader);
+	Shaders.TryGetPixelShader(OutPixelShader);
+
+	return true;
+}
+
 static void GetMaterialInfo(
 	const UMaterialInterface* InMaterialInterface,
 	ERHIFeatureLevel::Type InFeatureLevel,
@@ -367,19 +402,8 @@ static void GetMaterialInfo(
 		Material = MaterialProxy->GetMaterialNoFallback(InFeatureLevel);
 		if (Material && Material->GetMaterialDomain() == MD_PostProcess)
 		{
-			FMaterialShaderTypes ShaderTypes;
-			{
-				const bool bManualStencilTest = Inputs.bManualStencilTest && Material->IsStencilTestEnabled();
-
-				FPostProcessMaterialPS::FPermutationDomain PermutationVectorPS;
-				PermutationVectorPS.Set<FPostProcessMaterialPS::FManualStencilTestDim>(bManualStencilTest);
-				PermutationVectorPS.Set<FPostProcessMaterialPS::FNeuralPostProcessPrePass>(bNeuralPostProcessPrepass);
-
-				ShaderTypes.AddShaderType<FPostProcessMaterialVS>();
-				ShaderTypes.AddShaderType<FPostProcessMaterialPS>(PermutationVectorPS.ToDimensionValueId());
-			}
-
-			if (Material->TryGetShaders(ShaderTypes, nullptr, Shaders))
+			const bool bManualStencilTest = Inputs.bManualStencilTest && Material->IsStencilTestEnabled();
+			if (GetMaterialShaders(*Material, bManualStencilTest, bNeuralPostProcessPrepass, OutVertexShader, OutPixelShader))
 			{
 				break;
 			}
@@ -395,8 +419,6 @@ static void GetMaterialInfo(
 	OutMaterial = Material;
 	OutMaterialProxy = MaterialProxy;
 	OutMaterialShaderMap = MaterialShaderMap;
-	Shaders.TryGetVertexShader(OutVertexShader);
-	Shaders.TryGetPixelShader(OutPixelShader);
 }
 
 TGlobalResource<FPostProcessMaterialVertexDeclaration> GPostProcessMaterialVertexDeclaration;
@@ -1495,3 +1517,100 @@ FScreenPassTexture FPostProcessMaterialInputs::ReturnUntouchedSceneColorForPostP
 		return FScreenPassTexture::CopyFromSlice(GraphBuilder, SceneColorSlice, OverrideOutput);
 	}
 }
+
+static const TCHAR* PostProcessMaterialPSOCollectorName = TEXT("PostProcessMaterialPSOCollector");
+
+class FPostProcessMaterialPSOCollector : public IPSOCollector
+{
+public:
+	FPostProcessMaterialPSOCollector(ERHIFeatureLevel::Type InFeatureLevel) : 
+	IPSOCollector(FPSOCollectorCreateManager::GetIndex(GetFeatureLevelShadingPath(InFeatureLevel), PostProcessMaterialPSOCollectorName)),
+	FeatureLevel(InFeatureLevel)
+	{
+	}
+
+	virtual void CollectPSOInitializers(
+		const FSceneTexturesConfig& SceneTexturesConfig,
+		const FMaterial& Material,
+		const FPSOPrecacheVertexFactoryData& VertexFactoryData,
+		const FPSOPrecacheParams& PreCacheParams,
+		TArray<FPSOPrecacheData>& PSOInitializers
+	) override final;
+
+private:
+
+	ERHIFeatureLevel::Type FeatureLevel;
+};
+
+void FPostProcessMaterialPSOCollector::CollectPSOInitializers(
+	const FSceneTexturesConfig& SceneTexturesConfig,
+	const FMaterial& Material,
+	const FPSOPrecacheVertexFactoryData& VertexFactoryData,
+	const FPSOPrecacheParams& PreCacheParams,
+	TArray<FPSOPrecacheData>& PSOInitializers
+)
+{
+	if (!Material.IsPostProcessMaterial() || GPostProcessingMaterialPSOPrecache == 0)
+	{
+		return;
+	}
+
+	const FMaterialShaderMap* MaterialShaderMap = Material.GetGameThreadShaderMap();	
+	auto AddPSOInitializer = [&](bool bManualStencilTest)
+	{
+		TShaderRef<FPostProcessMaterialVS> VertexShader;
+		TShaderRef<FPostProcessMaterialPS> PixelShader;
+
+		bool bNeuralPostProcessPrepass = false;
+		if (!GetMaterialShaders(Material, bManualStencilTest, bNeuralPostProcessPrepass, VertexShader, PixelShader))
+		{
+			return;
+		}
+		      			
+		FRHIBlendState* DefaultBlendState = FScreenPassPipelineState::FDefaultBlendState::GetRHI();
+		FRHIBlendState* BlendState = IsMaterialBlendEnabled(&Material) ? GetMaterialBlendState(&Material) : DefaultBlendState;
+	
+		FRHIDepthStencilState* DefaultDepthStencilState = FScreenPassPipelineState::FDefaultDepthStencilState::GetRHI();
+		FRHIDepthStencilState* DepthStencilState = bManualStencilTest ? GetMaterialStencilState(&Material) : DefaultDepthStencilState;
+		
+		FGraphicsPipelineStateInitializer GraphicsPSOInit;
+		GraphicsPSOInit.BlendState = BlendState;
+		GraphicsPSOInit.RasterizerState = TStaticRasterizerState<FM_Solid, CM_None>::GetRHI();
+		GraphicsPSOInit.DepthStencilState = DepthStencilState;
+		GraphicsPSOInit.BoundShaderState.VertexDeclarationRHI = GFilterVertexDeclaration.VertexDeclarationRHI;
+		GraphicsPSOInit.BoundShaderState.VertexShaderRHI = VertexShader.GetVertexShader();
+		GraphicsPSOInit.BoundShaderState.PixelShaderRHI = PixelShader.GetPixelShader();
+		GraphicsPSOInit.PrimitiveType = PT_TriangleList;
+			
+		// What render target formats to support?
+		FGraphicsPipelineRenderTargetsInfo RenderTargetsInfo;
+		RenderTargetsInfo.NumSamples = 1;
+		AddRenderTargetInfo(SceneTexturesConfig.ColorFormat, SceneTexturesConfig.ColorCreateFlags, RenderTargetsInfo);
+		
+		GraphicsPSOInit.StatePrecachePSOHash = RHIComputeStatePrecachePSOHash(GraphicsPSOInit);
+		ApplyTargetsInfo(GraphicsPSOInit, RenderTargetsInfo);
+
+		FPSOPrecacheData PSOPrecacheData;
+		PSOPrecacheData.bRequired = true;
+		PSOPrecacheData.Type = FPSOPrecacheData::EType::Graphics;
+		PSOPrecacheData.GraphicsPSOInitializer = GraphicsPSOInit;
+#if PSO_PRECACHING_VALIDATE
+		PSOPrecacheData.PSOCollectorIndex = PSOCollectorIndex;
+		PSOPrecacheData.VertexFactoryType = nullptr;
+#endif // PSO_PRECACHING_VALIDATE
+
+		PSOInitializers.Add(MoveTemp(PSOPrecacheData));
+	};
+	
+	AddPSOInitializer(false /*bManualStencilTest*/);
+	if (Material.IsStencilTestEnabled())
+	{
+		AddPSOInitializer(true /*bManualStencilTest*/);
+	}
+}
+
+IPSOCollector* CreatePostProcessMaterialPSOCollector(ERHIFeatureLevel::Type FeatureLevel)
+{
+	return new FPostProcessMaterialPSOCollector(FeatureLevel);
+}
+FRegisterPSOCollectorCreateFunction RegisterPostProcessMaterialPSOCollector(&CreatePostProcessMaterialPSOCollector, EShadingPath::Deferred, PostProcessMaterialPSOCollectorName);
diff --git a/Engine/Source/Runtime/Renderer/Private/TranslucentLighting.cpp b/Engine/Source/Runtime/Renderer/Private/TranslucentLighting.cpp
--- a/Engine/Source/Runtime/Renderer/Private/TranslucentLighting.cpp
+++ b/Engine/Source/Runtime/Renderer/Private/TranslucentLighting.cpp
@@ -154,6 +154,15 @@ static FAutoConsoleVariableRef CVarTranslucencyLightingVolumeDebug(
 	ECVF_RenderThreadSafe
 );
 
+static int32 GTranslucencyLightingVolumeMaterialPSOPrecache = 1;
+static FAutoConsoleVariableRef CVarGTranslucencyLightingVolumeMaterialPSOPrecache(
+	TEXT("r.PSOPrecache.TranslucencyLightingVolumeMaterial"), 
+	GTranslucencyLightingVolumeMaterialPSOPrecache, 
+	TEXT("Precache all possible required Translucency Lighting Volume PSOs for loaded LightMaterials."),
+	ECVF_ReadOnly);
+
+static const TCHAR* TranslucentLightingMaterialPSOCollectorName = TEXT("TranslucentLightingMaterialPSOCollector");
+
 /** Function returning current translucency lighting volume dimensions. */
 int32 GetTranslucencyLightingVolumeDim()
 {
@@ -878,6 +887,15 @@ public:
 
 IMPLEMENT_GLOBAL_SHADER(FClearTranslucentLightingVolumeCS, "/Engine/Private/TranslucentLightInjectionShaders.usf", "ClearTranslucentLightingVolumeCS", SF_Compute);
 
+void FTranslucencyLightingVolumeTextures::GetTextureFormatAndCreationFlags(EPixelFormat& OutPixelFormat, ETextureCreateFlags& OutCreationFlags)
+{
+	OutPixelFormat = PF_FloatRGBA;
+
+	// TODO: We can skip the and TLV allocations when rendering in forward shading mode
+	OutCreationFlags = TexCreate_ShaderResource | TexCreate_RenderTargetable | TexCreate_ReduceMemoryWithTilingMode | TexCreate_UAV;
+}
+
+
 int32 FTranslucencyLightingVolumeTextures::GetIndex(const FViewInfo& View, int32 CascadeIndex) const
 {
 	// if we only have one view or one stereo pair we can just use primary index
@@ -909,6 +927,10 @@ void FTranslucencyLightingVolumeTextures::Init(FRDGBuilder& GraphBuilder, TArray
 	VolumeDim = GetTranslucencyLightingVolumeDim();
 	const FIntVector TranslucencyLightingVolumeDim(VolumeDim);
 
+	EPixelFormat TranslucencyPixelFormat;
+	ETextureCreateFlags TranslucencyTargetFlags;
+	GetTextureFormatAndCreationFlags(TranslucencyPixelFormat, TranslucencyTargetFlags);
+
 	// calculate the number of textures needed given that for each stereo pair the primary view's textures will be shared between the "eyes"
 	const int32 ViewCount = Views.Num();
 	uint32 NumViewsWithTextures = 0;
@@ -920,9 +942,6 @@ void FTranslucencyLightingVolumeTextures::Init(FRDGBuilder& GraphBuilder, TArray
 	}
 	check(NumViewsWithTextures > 0);
 	{
-		// TODO: We can skip the and TLV allocations when rendering in forward shading mode
-		const ETextureCreateFlags TranslucencyTargetFlags = TexCreate_ShaderResource | TexCreate_RenderTargetable | TexCreate_ReduceMemoryWithTilingMode | TexCreate_UAV;
-
 		Ambient.SetNum(NumViewsWithTextures * TVC_MAX);
 		Directional.SetNum(NumViewsWithTextures * TVC_MAX);
 
@@ -936,7 +955,7 @@ void FTranslucencyLightingVolumeTextures::Init(FRDGBuilder& GraphBuilder, TArray
 				FRDGTextureRef AmbientTexture = GraphBuilder.CreateTexture(
 					FRDGTextureDesc::Create3D(
 						TranslucencyLightingVolumeDim,
-						PF_FloatRGBA,
+						TranslucencyPixelFormat,
 						FClearValueBinding::Transparent,
 						TranslucencyTargetFlags),
 					TEXT("TranslucentVolumeAmbient"));
@@ -944,7 +963,7 @@ void FTranslucencyLightingVolumeTextures::Init(FRDGBuilder& GraphBuilder, TArray
 				FRDGTextureRef DirectionalTexture = GraphBuilder.CreateTexture(
 					FRDGTextureDesc::Create3D(
 						TranslucencyLightingVolumeDim,
-						PF_FloatRGBA,
+						TranslucencyPixelFormat,
 						FClearValueBinding::Transparent,
 						TranslucencyTargetFlags),
 					TEXT("TranslucentVolumeDirectional"));
@@ -1265,18 +1284,13 @@ static FRDGTextureRef GetSkyTransmittanceLutTexture(FRDGBuilder& GraphBuilder, c
 	return TransmittanceLutTexture;
 }
 
-
-static void SetPSOStateForVolumeInjection(
+static void SetupPSOStateForVolumeInjection(
 	TShaderMapRef<FWriteToSliceVS> VertexShader,
 	TOptionalShaderMapRef<FWriteToSliceGS> GeometryShader,
 	FRHIPixelShader* PixelShaderRHI,
-	FVolumeBounds VolumeBounds,
 	bool bDirectionalLight,
-	FRHICommandList& RHICmdList,
 	FGraphicsPipelineStateInitializer &OutGraphicsPSOInit)
 {
-	RHICmdList.ApplyCachedRenderTargets(OutGraphicsPSOInit);
-
 	OutGraphicsPSOInit.RasterizerState = TStaticRasterizerState<FM_Solid, CM_None>::GetRHI();
 	OutGraphicsPSOInit.DepthStencilState = TStaticDepthStencilState<false, CF_Always>::GetRHI();
 	OutGraphicsPSOInit.PrimitiveType = PT_TriangleStrip;
@@ -1287,19 +1301,34 @@ static void SetPSOStateForVolumeInjection(
 		// Directional lights write their shadowing into alpha of the ambient texture
 		OutGraphicsPSOInit.BlendState = TStaticBlendState<
 			CW_RGBA, BO_Add, BF_One, BF_One, BO_Add, BF_One, BF_One,
-			CW_RGB, BO_Add, BF_One, BF_One, BO_Add, BF_One, BF_One>::GetRHI();
+		CW_RGB, BO_Add, BF_One, BF_One, BO_Add, BF_One, BF_One>::GetRHI();
 	}
 	else
 	{
 		OutGraphicsPSOInit.BlendState = TStaticBlendState<
 			CW_RGB, BO_Add, BF_One, BF_One, BO_Add, BF_Zero, BF_One,
-			CW_RGB, BO_Add, BF_One, BF_One, BO_Add, BF_Zero, BF_One>::GetRHI();
+		CW_RGB, BO_Add, BF_One, BF_One, BO_Add, BF_Zero, BF_One>::GetRHI();
 	}
 
 	OutGraphicsPSOInit.BoundShaderState.VertexDeclarationRHI = GScreenVertexDeclaration.VertexDeclarationRHI;
 	OutGraphicsPSOInit.BoundShaderState.VertexShaderRHI = VertexShader.GetVertexShader();
 	OutGraphicsPSOInit.BoundShaderState.SetGeometryShader(GeometryShader.GetGeometryShader());
 	OutGraphicsPSOInit.BoundShaderState.PixelShaderRHI = PixelShaderRHI;
+}
+
+static void SetPSOStateForVolumeInjection(
+	TShaderMapRef<FWriteToSliceVS> VertexShader,
+	TOptionalShaderMapRef<FWriteToSliceGS> GeometryShader,
+	FRHIPixelShader* PixelShaderRHI,
+	FVolumeBounds VolumeBounds,
+	bool bDirectionalLight,
+	FRHICommandList& RHICmdList,
+	FGraphicsPipelineStateInitializer &OutGraphicsPSOInit)
+{
+	RHICmdList.ApplyCachedRenderTargets(OutGraphicsPSOInit);
+
+	SetupPSOStateForVolumeInjection(VertexShader, GeometryShader, PixelShaderRHI, bDirectionalLight, OutGraphicsPSOInit);
+
 	SetGraphicsPipelineState(RHICmdList, OutGraphicsPSOInit, 0);
 
 	const int32 TranslucencyLightingVolumeDim = GetTranslucencyLightingVolumeDim();
@@ -1547,6 +1576,14 @@ void InjectTranslucencyLightingVolume(
 					SetPSOStateForVolumeInjection(VertexShader, GeometryShader, PixelShader.GetPixelShader(),
 						VolumeBounds, bDirectionalLight, RHICmdList, GraphicsPSOInit);
 
+#if PSO_PRECACHING_VALIDATE
+					if (PSOCollectorStats::IsFullPrecachingValidationEnabled())
+					{
+						static const int32 MaterialPSOCollectorIndex = FPSOCollectorCreateManager::GetIndex(GetFeatureLevelShadingPath(GMaxRHIFeatureLevel), TranslucentLightingMaterialPSOCollectorName);
+						PSOCollectorStats::CheckFullPipelineStateInCache(GraphicsPSOInit, EPSOPrecacheResult::Unknown, MaterialProxy, nullptr, nullptr, MaterialPSOCollectorIndex);
+					}
+#endif // PSO_PRECACHING_VALIDATE
+
 					SetShaderParametersLegacyPS(RHICmdList, PixelShader, View, InjectionData.LightFunctionMaterialProxy);
 					SetShaderParameters(RHICmdList, PixelShader, PixelShader.GetPixelShader(), PassParameters->PS);
 
@@ -1822,3 +1859,113 @@ void FilterTranslucencyLightingVolume(
 		}
 	}
 }
+
+class FTranslucentLightingMaterialPSOCollector : public IPSOCollector
+{
+public:
+	FTranslucentLightingMaterialPSOCollector(ERHIFeatureLevel::Type InFeatureLevel) : 
+		IPSOCollector(FPSOCollectorCreateManager::GetIndex(GetFeatureLevelShadingPath(InFeatureLevel), TranslucentLightingMaterialPSOCollectorName)),
+	FeatureLevel(InFeatureLevel)
+	{
+	}
+
+	virtual void CollectPSOInitializers(
+		const FSceneTexturesConfig& SceneTexturesConfig,
+		const FMaterial& Material,
+		const FPSOPrecacheVertexFactoryData& VertexFactoryData,
+		const FPSOPrecacheParams& PreCacheParams,
+		TArray<FPSOPrecacheData>& PSOInitializers
+	) override final;
+
+private:
+
+	ERHIFeatureLevel::Type FeatureLevel;
+};
+
+void FTranslucentLightingMaterialPSOCollector::CollectPSOInitializers(
+	const FSceneTexturesConfig& SceneTexturesConfig,
+	const FMaterial& Material,
+	const FPSOPrecacheVertexFactoryData& VertexFactoryData,
+	const FPSOPrecacheParams& PreCacheParams,
+	TArray<FPSOPrecacheData>& PSOInitializers
+)
+{
+	if (!Material.IsLightFunction() || GTranslucencyLightingVolumeMaterialPSOPrecache == 0)
+	{
+		return;
+	}
+
+	FGlobalShaderMap* GlobalShaderMap = GetGlobalShaderMap(FeatureLevel);
+	TShaderMapRef<FWriteToSliceVS> VertexShader(GlobalShaderMap);
+	TOptionalShaderMapRef<FWriteToSliceGS> GeometryShader(GlobalShaderMap);
+	const FMaterialShaderMap* MaterialShaderMap = Material.GetGameThreadShaderMap();
+	
+	EPixelFormat TranslucencyPixelFormat;
+	ETextureCreateFlags TranslucencyTargetFlags;
+	FTranslucencyLightingVolumeTextures::GetTextureFormatAndCreationFlags(TranslucencyPixelFormat, TranslucencyTargetFlags);
+
+	auto AddPSOInitializer = [&](bool bDirectionalLight, bool bDynamicShadow, bool bApplyLightFunction, bool bUseVSM, bool bUseAdaptiveVolumetricShadowMap)
+		{
+			FTranslucentLightingInjectPS::FPermutationDomain PermutationVector;
+			PermutationVector.Set< FTranslucentLightingInjectPS::FRadialAttenuation >(!bDirectionalLight);
+			PermutationVector.Set< FTranslucentLightingInjectPS::FDynamicallyShadowed >(bDynamicShadow);
+			PermutationVector.Set< FTranslucentLightingInjectPS::FLightFunction >(bApplyLightFunction);
+			PermutationVector.Set< FTranslucentLightingInjectPS::FVirtualShadowMap >(bUseVSM);
+			PermutationVector.Set< FTranslucentLightingInjectPS::FAdaptiveVolumetricShadowMap >(bUseAdaptiveVolumetricShadowMap);
+
+			auto PixelShader = MaterialShaderMap->GetShader< FTranslucentLightingInjectPS >(PermutationVector);
+			
+			FGraphicsPipelineStateInitializer GraphicsPSOInit;
+			SetupPSOStateForVolumeInjection(VertexShader, GeometryShader, PixelShader.GetPixelShader(), bDirectionalLight, GraphicsPSOInit);
+			
+			// What render target formats to support?
+			FGraphicsPipelineRenderTargetsInfo RenderTargetsInfo;
+			RenderTargetsInfo.NumSamples = 1;
+			AddRenderTargetInfo(TranslucencyPixelFormat, TranslucencyTargetFlags, RenderTargetsInfo);
+			AddRenderTargetInfo(TranslucencyPixelFormat, TranslucencyTargetFlags, RenderTargetsInfo);
+
+			GraphicsPSOInit.StatePrecachePSOHash = RHIComputeStatePrecachePSOHash(GraphicsPSOInit);
+			ApplyTargetsInfo(GraphicsPSOInit, RenderTargetsInfo);
+
+			FPSOPrecacheData PSOPrecacheData;
+			PSOPrecacheData.bRequired = true;
+			PSOPrecacheData.Type = FPSOPrecacheData::EType::Graphics;
+			PSOPrecacheData.GraphicsPSOInitializer = GraphicsPSOInit;
+#if PSO_PRECACHING_VALIDATE
+			PSOPrecacheData.PSOCollectorIndex = PSOCollectorIndex;
+			PSOPrecacheData.VertexFactoryType = nullptr;
+#endif // PSO_PRECACHING_VALIDATE
+
+			PSOInitializers.Add(MoveTemp(PSOPrecacheData));
+		};
+	
+	// Generate PSOs for all possible permutations - we don't have that many light function materials
+	for (int32 DirectionalLightIndex = 0; DirectionalLightIndex < 2; ++DirectionalLightIndex)
+	{
+		bool bDirectionalLight = DirectionalLightIndex > 0;
+		for (int32 DynamicShadowIndex = 0; DynamicShadowIndex < 2; ++DynamicShadowIndex)
+		{
+			bool bDynamicShadow = DynamicShadowIndex > 0;
+			for (int32 ApplyLightFunctionIndex = 0; ApplyLightFunctionIndex < 2; ++ApplyLightFunctionIndex)
+			{
+				bool bApplyLightFunction = ApplyLightFunctionIndex > 0;
+				for (int32 UseVSMIndex = 0; UseVSMIndex < 2; ++UseVSMIndex)
+				{
+					bool bUseVSM = UseVSMIndex > 0;
+					for (int32 UseAdaptiveVolumetricShadowMapIndex = 0; UseAdaptiveVolumetricShadowMapIndex < 2; ++UseAdaptiveVolumetricShadowMapIndex)
+					{
+						bool bUseAdaptiveVolumetricShadowMap = UseAdaptiveVolumetricShadowMapIndex > 0;
+						AddPSOInitializer(bDirectionalLight, bDynamicShadow, bApplyLightFunction, bUseVSM, bUseAdaptiveVolumetricShadowMap);
+					}
+				}
+			}
+		}
+	}
+}
+
+IPSOCollector* CreateTranslucentLightingMaterialPSOCollector(ERHIFeatureLevel::Type FeatureLevel)
+{
+	return new FTranslucentLightingMaterialPSOCollector(FeatureLevel);
+}
+FRegisterPSOCollectorCreateFunction RegisterTranslucentLightingMaterialPSOCollector(&CreateTranslucentLightingMaterialPSOCollector, EShadingPath::Deferred, TranslucentLightingMaterialPSOCollectorName);
+
diff --git a/Engine/Source/Runtime/Renderer/Private/TranslucentLighting.h b/Engine/Source/Runtime/Renderer/Private/TranslucentLighting.h
--- a/Engine/Source/Runtime/Renderer/Private/TranslucentLighting.h
+++ b/Engine/Source/Runtime/Renderer/Private/TranslucentLighting.h
@@ -14,6 +14,9 @@ class FProjectedShadowInfo;
 
 struct FTranslucencyLightingVolumeTextures
 {
+	/** Retrieve the statically defined pixel format and creation flags used for the translucency light volume textures */
+	static void GetTextureFormatAndCreationFlags(EPixelFormat& OutPixelFormat, ETextureCreateFlags& OutCreationFlags);
+
 	/** Maps view to a texture index, accounting for sharing between some views */
 	int32 GetIndex(const FViewInfo& View, int32 CascadeIndex) const;
 
diff --git a/Engine/Source/Runtime/SlateRHIRenderer/Private/SlateRHIRenderingPolicy.cpp b/Engine/Source/Runtime/SlateRHIRenderer/Private/SlateRHIRenderingPolicy.cpp
--- a/Engine/Source/Runtime/SlateRHIRenderer/Private/SlateRHIRenderingPolicy.cpp
+++ b/Engine/Source/Runtime/SlateRHIRenderer/Private/SlateRHIRenderingPolicy.cpp
@@ -34,6 +34,7 @@
 #include "SceneInterface.h"
 #include "Containers/StaticBitArray.h"
 #include "MeshPassProcessor.h"
+#include "PSOPrecacheValidation.h"
 
 extern void UpdateNoiseTextureParameters(FViewUniformShaderParameters& ViewUniformShaderParameters);
 
@@ -47,6 +48,16 @@ DECLARE_DWORD_COUNTER_STAT(TEXT("Num Vertices"), STAT_SlateVertexCount, STATGROU
 DECLARE_DWORD_COUNTER_STAT(TEXT("Clips (Scissor)"), STAT_SlateScissorClips, STATGROUP_Slate);
 DECLARE_DWORD_COUNTER_STAT(TEXT("Clips (Stencil)"), STAT_SlateStencilClips, STATGROUP_Slate);
 
+static int32 GSlateMaterialPSOPrecache = 1;
+static FAutoConsoleVariableRef CVarGSlateMaterialPSOPrecache(
+	TEXT("r.PSOPrecache.SlateMaterials"), 
+	GSlateMaterialPSOPrecache, 
+	TEXT("Precache all possible required PSOs for loaded Slate Materials."),
+	ECVF_ReadOnly);
+
+static const TCHAR* SlateGlobalPSOCollectorName = TEXT("SlateGlobalPSOCollector");
+static const TCHAR* SlateMaterialPSOCollectorName = TEXT("SlateMaterialPSOCollector");
+
 #if WITH_SLATE_DEBUGGING
 static int32 GSlateEnableDrawEvents = 1;
 #else
@@ -225,7 +236,7 @@ TShaderRef<FSlateElementPS> GetTexturePixelShader(FGlobalShaderMap* ShaderMap, E
 	TShaderRef<FSlateElementPS> PixelShader;
 
 #if WITH_SLATE_VISUALIZERS
-	if (CVarShowSlateOverdraw.GetValueOnRenderThread() != 0)
+	if (CVarShowSlateOverdraw.GetValueOnAnyThread() != 0)
 	{
 		PixelShader = TShaderMapRef<FSlateDebugOverdrawPS>(ShaderMap);
 	}
@@ -414,7 +425,7 @@ TShaderRef<FSlateElementPS> GetTexturePixelShader(FGlobalShaderMap* ShaderMap, E
 	return PixelShader;
 }
 
-void ChooseMaterialShaderTypes(ESlateShader ShaderType, bool bUseInstancing, FMaterialShaderTypes& OutShaderTypes)
+bool ChooseMaterialShaderTypes(ESlateShader ShaderType, bool bUseInstancing, FMaterialShaderTypes& OutShaderTypes)
 {
 	switch (ShaderType)
 	{
@@ -443,8 +454,7 @@ void ChooseMaterialShaderTypes(ESlateShader ShaderType, bool bUseInstancing, FMa
 		OutShaderTypes.AddShaderType<TSlateMaterialShaderPS<ESlateShader::MsdfFont>>();
 		break;
 	default:
-		checkf(false, TEXT("Unsupported Slate shader type for use with materials"));
-		break;
+		return false;
 	}
 
 	if (bUseInstancing)
@@ -455,6 +465,8 @@ void ChooseMaterialShaderTypes(ESlateShader ShaderType, bool bUseInstancing, FMa
 	{
 		OutShaderTypes.AddShaderType<TSlateMaterialShaderVS<false>>();
 	}
+
+	return true;
 }
 
 EPrimitiveType GetRHIPrimitiveType(ESlateDrawPrimitive SlateType)
@@ -1064,7 +1076,11 @@ FSlateRenderBatchOp* CreateSlateRenderBatchOp(
 		TShaderRef<FSlateMaterialShaderPS> PixelShader;
 
 		FMaterialShaderTypes ShaderTypesToGet;
-		ChooseMaterialShaderTypes(ShaderType, bUseInstancing, ShaderTypesToGet);
+		if (!ChooseMaterialShaderTypes(ShaderType, bUseInstancing, ShaderTypesToGet))
+		{
+			checkf(false, TEXT("Unsupported Slate shader type for use with materials"));
+			return nullptr;
+		}
 		const FMaterial* EffectiveMaterial = nullptr;
 
 		while (MaterialRenderProxy)
@@ -1292,6 +1308,15 @@ void DrawSlateRenderBatch(
 		State.GraphicsPSOInit.BoundShaderState.VertexShaderRHI = RenderBatchOp.VertexBindings->Shader.GetVertexShader();
 		State.GraphicsPSOInit.BoundShaderState.PixelShaderRHI  = RenderBatchOp.PixelBindings->Shader.GetPixelShader();
 		State.GraphicsPSOInit.PrimitiveType = GetRHIPrimitiveType(RenderBatch.DrawPrimitiveType);
+		
+#if PSO_PRECACHING_VALIDATE
+		if (PSOCollectorStats::IsFullPrecachingValidationEnabled())
+		{
+			static const int32 MaterialPSOCollectorIndex = FPSOCollectorCreateManager::GetIndex(GetFeatureLevelShadingPath(GMaxRHIFeatureLevel), SlateMaterialPSOCollectorName);
+			// Material Render Proxy is not cached in the state, but could be done for better debugging if needed
+			PSOCollectorStats::CheckFullPipelineStateInCache(State.GraphicsPSOInit, EPSOPrecacheResult::Unknown, nullptr, nullptr, nullptr, MaterialPSOCollectorIndex);
+		}
+#endif // PSO_PRECACHING_VALIDATE
 
 		SetGraphicsPipelineState(RHICmdList, State.GraphicsPSOInit, State.StencilRef);
 
@@ -1334,6 +1359,14 @@ void DrawSlateRenderBatch(
 		State.GraphicsPSOInit.BoundShaderState.PixelShaderRHI = RenderBatchOp.PixelBindings->Shader.GetPixelShader();
 		State.GraphicsPSOInit.PrimitiveType = GetRHIPrimitiveType(RenderBatch.DrawPrimitiveType);
 
+#if PSO_PRECACHING_VALIDATE
+		if (PSOCollectorStats::IsFullPrecachingValidationEnabled())
+		{			
+			static const int32 GlobalPSOCollectorIndex = FGlobalPSOCollectorManager::GetIndex(SlateGlobalPSOCollectorName);
+			PSOCollectorStats::CheckGlobalGraphicsPipelineStateInCache(State.GraphicsPSOInit, GlobalPSOCollectorIndex);
+		}
+#endif // PSO_PRECACHING_VALIDATE
+
 		SetGraphicsPipelineState(RHICmdList, State.GraphicsPSOInit, State.StencilRef);
 
 		RenderBatchOp.PixelBindings->SetOnCommandList(RHICmdList);
@@ -1651,4 +1684,196 @@ void AddSlateDrawElementsPass(
 	
 	INC_DWORD_STAT_BY(STAT_SlateScissorClips, ClippingCreateContext.NumScissors);
 	INC_DWORD_STAT_BY(STAT_SlateStencilClips, ClippingCreateContext.NumStencils);
-}
\ No newline at end of file
+}
+
+static TArray<ESlateShader> SlateShaderTypesToPrecache =
+{
+	ESlateShader::Default,
+	ESlateShader::Border,
+	ESlateShader::GrayscaleFont,
+	ESlateShader::Custom,
+	ESlateShader::RoundedBox
+};
+
+void AddSlatePSOInitializer(
+	FRHIBlendState* BlendState, 
+	bool bInstanced, 
+	ESlateDrawPrimitive DrawPrimitiveType, 
+	FRHIVertexShader* VertexShaderRHI, 
+	FRHIPixelShader* PixelShaderRHI,
+	const FSceneTexturesConfig& SceneTexturesConfig, 
+	int32 PSOCollectorIndex, 
+	TArray<FPSOPrecacheData>& PSOInitializers)
+{
+	check(VertexShaderRHI && PixelShaderRHI);
+
+	FGraphicsPipelineStateInitializer GraphicsPSOInit;	
+
+	GraphicsPSOInit.DepthStencilState = TStaticDepthStencilState<false, CF_Always>::GetRHI();
+	GraphicsPSOInit.RasterizerState = TStaticRasterizerState<FM_Solid>::GetRHI();
+	GraphicsPSOInit.BlendState = BlendState;
+
+	GraphicsPSOInit.BoundShaderState.VertexDeclarationRHI = bInstanced ? GSlateInstancedVertexDeclaration.VertexDeclarationRHI : GSlateVertexDeclaration.VertexDeclarationRHI;
+	GraphicsPSOInit.BoundShaderState.VertexShaderRHI = VertexShaderRHI;
+	GraphicsPSOInit.BoundShaderState.PixelShaderRHI = PixelShaderRHI;
+	GraphicsPSOInit.PrimitiveType = GetRHIPrimitiveType(DrawPrimitiveType);
+
+	FGraphicsPipelineRenderTargetsInfo RenderTargetsInfo;
+	AddRenderTargetInfo(PF_A2B10G10R10, ETextureCreateFlags::RenderTargetable, RenderTargetsInfo);
+	RenderTargetsInfo.NumSamples = 1;
+
+	GraphicsPSOInit.StatePrecachePSOHash = RHIComputeStatePrecachePSOHash(GraphicsPSOInit);		
+	ApplyTargetsInfo(GraphicsPSOInit, RenderTargetsInfo);
+
+	FPSOPrecacheData PSOPrecacheData;
+	PSOPrecacheData.bRequired = true;
+	PSOPrecacheData.Type = FPSOPrecacheData::EType::Graphics;
+	PSOPrecacheData.GraphicsPSOInitializer = GraphicsPSOInit;
+#if PSO_PRECACHING_VALIDATE
+	PSOPrecacheData.PSOCollectorIndex = PSOCollectorIndex;
+	PSOPrecacheData.VertexFactoryType = nullptr;
+#endif // PSO_PRECACHING_VALIDATE
+
+	PSOInitializers.Add(MoveTemp(PSOPrecacheData));
+}
+
+void SlateGlobalPSOCollector(const FSceneTexturesConfig& SceneTexturesConfig, int32 GlobalPSOCollectorIndex, TArray<FPSOPrecacheData>& PSOInitializers)
+{
+	EShaderPlatform ShaderPlatform = SceneTexturesConfig.ShaderPlatform;
+	FGlobalShaderMap* GlobalShaderMap = GetGlobalShaderMap(ShaderPlatform);
+
+	bool bRequired = false;
+
+	// Collect all possible permutations of the FSlateElementPS (out of 5K possible permutations only 10 unique Pixel shaders are found)	
+	TSet<FRHIPixelShader*> SlateElementPixelShaders;
+	for (int32 ShaderTypeIndex = 0; ShaderTypeIndex <= (int32)ESlateShader::MsdfFont; ++ShaderTypeIndex)
+	{
+		ESlateShader ShaderType = (ESlateShader) ShaderTypeIndex;
+		for (int32 DrawEffectsIndex = 0; DrawEffectsIndex <= (int32)ESlateDrawEffect::ReverseGamma; ++DrawEffectsIndex)
+		{
+			ESlateDrawEffect DrawEffects = (ESlateDrawEffect) ShaderTypeIndex;
+			for (int32 UseTextureGreyScaleIndex = 0; UseTextureGreyScaleIndex < 2; ++UseTextureGreyScaleIndex)
+			{
+				bool bUseTextureGrayscale = UseTextureGreyScaleIndex > 0;
+				for (int32 IsVirtualTextureIndex = 0; IsVirtualTextureIndex < 2; ++IsVirtualTextureIndex)
+				{
+					bool bIsVirtualTexture = IsVirtualTextureIndex > 0;
+
+					TShaderRef<FSlateElementPS> SlateElementPS = GetTexturePixelShader(GlobalShaderMap, ShaderType, DrawEffects, bUseTextureGrayscale, bIsVirtualTexture);
+					FRHIPixelShader* RHIPixelShader = static_cast<FRHIPixelShader*>(SlateElementPS.GetRHIShaderBase(SF_Pixel, bRequired));
+					if (RHIPixelShader)
+					{
+						SlateElementPixelShaders.Add(RHIPixelShader);
+					}
+				}
+			}
+		}
+	}
+
+	TArray<FBlendStateRHIRef> SlateElementBlendStates;
+	SlateElementBlendStates.Add(TStaticBlendState<>::GetRHI());
+	SlateElementBlendStates.Add(TStaticBlendState<CW_RGBA, BO_Add, BF_One, BF_InverseSourceAlpha, BO_Add, BF_One, BF_InverseSourceAlpha>::GetRHI());
+	SlateElementBlendStates.Add(TStaticBlendState<CW_RGBA, BO_Add, BF_SourceAlpha, BF_InverseSourceAlpha, BO_Add, BF_One, BF_InverseSourceAlpha>::GetRHI());
+
+	TShaderMapRef<FSlateElementVS> SlateElementVertexShader(GlobalShaderMap);
+	FRHIVertexShader* RHIVertexShader = static_cast<FRHIVertexShader*>(SlateElementVertexShader.GetRHIShaderBase(SF_Vertex));
+
+	for (FRHIPixelShader* SlateElementPixelShader : SlateElementPixelShaders)
+	{
+		for (FBlendStateRHIRef& BlendState : SlateElementBlendStates)
+		{
+			bool bInstanced = false;
+			AddSlatePSOInitializer(BlendState, bInstanced, ESlateDrawPrimitive::TriangleList, RHIVertexShader, SlateElementPixelShader, SceneTexturesConfig, GlobalPSOCollectorIndex, PSOInitializers);
+		}
+	}
+}
+
+FRegisterGlobalPSOCollectorFunction RegisterSlateGlobalPSOCollector(&SlateGlobalPSOCollector, SlateGlobalPSOCollectorName);
+
+class FSlateMaterialPSOCollector : public IPSOCollector
+{
+public:
+	FSlateMaterialPSOCollector(ERHIFeatureLevel::Type InFeatureLevel) : 
+		IPSOCollector(FPSOCollectorCreateManager::GetIndex(GetFeatureLevelShadingPath(InFeatureLevel), SlateMaterialPSOCollectorName)),
+		FeatureLevel(InFeatureLevel)
+	{
+	}
+
+	virtual void CollectPSOInitializers(
+		const FSceneTexturesConfig& SceneTexturesConfig,
+		const FMaterial& Material,
+		const FPSOPrecacheVertexFactoryData& VertexFactoryData,
+		const FPSOPrecacheParams& PreCacheParams,
+		TArray<FPSOPrecacheData>& PSOInitializers
+	) override final;
+
+private:
+
+	ERHIFeatureLevel::Type FeatureLevel;
+};
+
+void FSlateMaterialPSOCollector::CollectPSOInitializers(
+	const FSceneTexturesConfig& SceneTexturesConfig,
+	const FMaterial& Material,
+	const FPSOPrecacheVertexFactoryData& VertexFactoryData,
+	const FPSOPrecacheParams& PreCacheParams,
+	TArray<FPSOPrecacheData>& PSOInitializers
+)
+{
+	if (!Material.IsUIMaterial() || GSlateMaterialPSOPrecache == 0)
+	{
+		return;
+	}
+
+	bool bRequired = false;	
+	for (ESlateShader ShaderType : SlateShaderTypesToPrecache)
+	{
+		int32 UseInstancingCount = ShaderType == ESlateShader::Custom ? 2 : 1;
+		for (int32 UseInstancingIndex = 0; UseInstancingIndex < UseInstancingCount; ++UseInstancingIndex)
+		{
+			bool bUseInstancing = UseInstancingIndex > 0;
+
+			FMaterialShaderTypes ShaderTypesToGet;
+			if (!ChooseMaterialShaderTypes(ShaderType, bUseInstancing, ShaderTypesToGet))
+			{
+				continue;
+			}
+
+			FMaterialShaders Shaders;
+			if (!Material.TryGetShaders(ShaderTypesToGet, nullptr, Shaders))
+			{
+				continue;
+			}			
+			
+			TShaderRef<FSlateMaterialShaderVS> VertexShader;
+			TShaderRef<FSlateMaterialShaderPS> PixelShader;
+			Shaders.TryGetVertexShader(VertexShader); 
+			Shaders.TryGetPixelShader(PixelShader);						
+			if (!VertexShader.IsValid() || !PixelShader.IsValid())
+			{
+				continue;
+			}
+
+			FRHIVertexShader* RHIVertexShader = VertexShader.GetVertexShader(bRequired);
+			FRHIPixelShader* RHIPixelShader = PixelShader.GetPixelShader(bRequired);
+			if (RHIVertexShader == nullptr || RHIPixelShader == nullptr)
+			{
+				continue;
+			}
+
+			// Don't know if mask resource will be used or not (also precache with blend mode required when mask resource is set?)
+			FSlateShaderResource* MaskResource = nullptr;
+			FRHIBlendState* BlendState = GetMaterialBlendState(MaskResource, &Material);
+
+			// Only precache TriangleList
+			AddSlatePSOInitializer(BlendState, bUseInstancing, ESlateDrawPrimitive::TriangleList, RHIVertexShader, RHIPixelShader, SceneTexturesConfig, PSOCollectorIndex, PSOInitializers);
+		}
+	}
+}
+
+IPSOCollector* CreateSlateMaterialPSOCollector(ERHIFeatureLevel::Type FeatureLevel)
+{
+	return new FSlateMaterialPSOCollector(FeatureLevel);
+}
+FRegisterPSOCollectorCreateFunction RegisterSlateMaterialPSOCollector(&CreateSlateMaterialPSOCollector, EShadingPath::Deferred, SlateMaterialPSOCollectorName);
+
