Added back slate material draw events.

diff --git a/Engine/Source/Runtime/SlateRHIRenderer/Private/SlateRHIRenderingPolicy.cpp b/Engine/Source/Runtime/SlateRHIRenderer/Private/SlateRHIRenderingPolicy.cpp
--- a/Engine/Source/Runtime/SlateRHIRenderer/Private/SlateRHIRenderingPolicy.cpp
+++ b/Engine/Source/Runtime/SlateRHIRenderer/Private/SlateRHIRenderingPolicy.cpp
@@ -65,7 +65,9 @@ static int32 GSlateEnableDrawEvents = 0;
 #endif
 static FAutoConsoleVariableRef CVarGSlateEnableDrawEvents(TEXT("Slate.EnableDrawEvents"), GSlateEnableDrawEvents, TEXT("."), ECVF_Default);
 
-#if !(UE_BUILD_SHIPPING || UE_BUILD_TEST)
+#define WITH_SLATE_DRAW_EVENTS !(UE_BUILD_SHIPPING || UE_BUILD_TEST)
+
+#if WITH_SLATE_DRAW_EVENTS
 	#define SLATE_DRAW_EVENT( RHICmdList, EventName             ) SCOPED_CONDITIONAL_DRAW_EVENT( RHICmdList, EventName, (GSlateEnableDrawEvents != 0));
 	#define SLATE_DRAW_EVENTF(RHICmdList, EventName, Format, ...) SCOPED_CONDITIONAL_DRAW_EVENTF(RHICmdList, EventName, (GSlateEnableDrawEvents != 0), Format, ##__VA_ARGS__);
 #else
@@ -1001,6 +1003,9 @@ struct FSlateRenderBatchOp
 	FRHIBuffer* InstanceBuffer;
 	FRHIBlendState* BlendState;
 	ESlateShaderResource::Type ShaderResourceType;
+#if WITH_SLATE_DRAW_EVENTS
+	const FString* MaterialName;
+#endif
 };
 
 //////////////////////////////////////////////////////////////////////////
@@ -1049,6 +1054,10 @@ FSlateRenderBatchOp* CreateSlateRenderBatchOp(
 	FSlateDrawShaderBindings* PixelBindings = nullptr;
 	FSlateDrawShaderBindings* VertexBindings = nullptr;
 
+#if WITH_SLATE_DRAW_EVENTS
+	const FString* MaterialName = nullptr;
+#endif
+
 	if (ResourceType == ESlateShaderResource::Material)
 	{
 		// Skip material render batches when the engine is not available.
@@ -1103,6 +1112,10 @@ FSlateRenderBatchOp* CreateSlateRenderBatchOp(
 			return nullptr;
 		}
 
+#if WITH_SLATE_DRAW_EVENTS
+		MaterialName = &MaterialRenderProxy->GetMaterialName();
+#endif
+
 		VertexBindings = FSlateDrawShaderBindings::Create(GraphBuilder, VertexShader);
 		VertexShader->SetMaterialShaderParameters(*VertexBindings, Scene, ViewUniformBuffer, MaterialRenderProxy, EffectiveMaterial);
 
@@ -1258,6 +1271,9 @@ FSlateRenderBatchOp* CreateSlateRenderBatchOp(
 	RenderBatchOp->InstanceBuffer     = bUseInstancing ? RenderBatch->InstanceData->GetRHI() : nullptr;
 	RenderBatchOp->BlendState         = BlendState;
 	RenderBatchOp->Next               = nullptr;
+#if WITH_SLATE_DRAW_EVENTS
+	RenderBatchOp->MaterialName       = MaterialName;
+#endif
 	return RenderBatchOp;
 }
 
@@ -1304,6 +1320,8 @@ void DrawSlateRenderBatch(
 
 	if (RenderBatchOp.ShaderResourceType == ESlateShaderResource::Material)
 	{
+		SLATE_DRAW_EVENTF(RHICmdList, MaterialBatch, TEXT("Slate Material: %s"), *RenderBatchOp.MaterialName);
+
 		State.GraphicsPSOInit.BoundShaderState.VertexDeclarationRHI = RenderBatchOp.InstanceBuffer ? GSlateInstancedVertexDeclaration.VertexDeclarationRHI : GSlateVertexDeclaration.VertexDeclarationRHI;
 		State.GraphicsPSOInit.BoundShaderState.VertexShaderRHI = RenderBatchOp.VertexBindings->Shader.GetVertexShader();
 		State.GraphicsPSOInit.BoundShaderState.PixelShaderRHI  = RenderBatchOp.PixelBindings->Shader.GetPixelShader();
