Skip rendering slate if the app doesn't have a swap chain.

diff --git a/Engine/Source/Runtime/SlateRHIRenderer/Private/SlateRHIRenderer.cpp b/Engine/Source/Runtime/SlateRHIRenderer/Private/SlateRHIRenderer.cpp
--- a/Engine/Source/Runtime/SlateRHIRenderer/Private/SlateRHIRenderer.cpp
+++ b/Engine/Source/Runtime/SlateRHIRenderer/Private/SlateRHIRenderer.cpp
@@ -1168,10 +1168,11 @@ void FSlateRHIRenderer::DrawWindows_Private(FSlateDrawBuffer& WindowDrawBuffer)
 
 	USlateRHIRendererSettings* RendererSettings = USlateRHIRendererSettings::GetMutable();
 
+	static const bool bRenderOffscreen = FParse::Param(FCommandLine::Get(), TEXT("RenderOffScreen"));
 	const float AppDeltaTime = FApp::GetDeltaTime();
 	const FGameTime AppDilatedTime = FGameTime::CreateDilated(FPlatformTime::Seconds() - GStartTime, AppDeltaTime, FApp::GetCurrentTime() - GStartTime, AppDeltaTime);
 	const FVector2f AppCursorPosition = FSlateApplication::Get().GetCursorPos();
-	const bool bAppCanRender = GIsClient && !IsRunningCommandlet() && !GUsingNullRHI;
+	const bool bAppCanRender = GIsClient && !IsRunningCommandlet() && !GUsingNullRHI && !bRenderOffscreen;
 	const bool bAppCanRenderPostProcess = RendererSettings && IsInGameThread() && UAssetManager::IsInitialized() && bAppCanRender && CVarCopyBackbufferToSlatePostRenderTargets.GetValueOnGameThread() > 0;
 	EPixelFormat AppViewportSceneFormat = PF_Unknown;
 	FIntPoint AppViewportExtentMax = FIntPoint::ZeroValue;
@@ -1196,102 +1197,106 @@ void FSlateRHIRenderer::DrawWindows_Private(FSlateDrawBuffer& WindowDrawBuffer)
 	};
 
 	TArray<FWindowToRender, FConcurrentLinearArrayAllocator> WindowsToRender;
-	WindowsToRender.Reserve(WindowDrawBuffer.GetWindowElementLists().Num());
 
-	for (const TSharedRef<FSlateWindowElementList>& WindowElementListRef : WindowDrawBuffer.GetWindowElementLists())
+	if (bAppCanRender)
 	{
-		FSlateWindowElementList* WindowElementList = &(*WindowElementListRef);
-		SWindow* Window = WindowElementList->GetRenderWindow();
+		WindowsToRender.Reserve(WindowDrawBuffer.GetWindowElementLists().Num());
 
-		if (!Window)
+		for (const TSharedRef<FSlateWindowElementList>& WindowElementListRef : WindowDrawBuffer.GetWindowElementLists())
 		{
-			ensureMsgf(false, TEXT("Window isn't valid but being drawn!"));
-			continue;
-		}
+			FSlateWindowElementList* WindowElementList = &(*WindowElementListRef);
+			SWindow* Window = WindowElementList->GetRenderWindow();
 
-		// This will return zero if both the viewport and the window are zero sized.
-		const FVector2f WindowSize = Window->GetViewportSize();
-		if (WindowSize.X <= 0.0f || WindowSize.Y <= 0.0f)
-		{
-			continue;
-		}
+			if (!Window)
+			{
+				ensureMsgf(false, TEXT("Window isn't valid but being drawn!"));
+				continue;
+			}
 
-		TRACE_CPUPROFILER_EVENT_SCOPE(GatherWindowElements);
+			// This will return zero if both the viewport and the window are zero sized.
+			const FVector2f WindowSize = Window->GetViewportSize();
+			if (WindowSize.X <= 0.0f || WindowSize.Y <= 0.0f)
+			{
+				continue;
+			}
 
-		// It's possible for a window to not have a viewport, in which case the viewport dimensions will be zero.
-		FVector2f ViewportCursorPosition = AppCursorPosition - Window->GetPositionInScreen();
-		FIntPoint ViewportOffset = FIntPoint::ZeroValue;
-		FIntPoint ViewportExtent = FIntPoint::ZeroValue;
-		FIntRect ViewportRect;
-		float ViewportScaleUI = Window->GetViewportScaleUIOverride();
+			TRACE_CPUPROFILER_EVENT_SCOPE(GatherWindowElements);
 
-		if (ISlateViewport* Viewport = Window->GetViewport().Get())
-		{
-			TSharedPtr<SWidget> ViewportWidget = Viewport->GetWidget().Pin();
+			// It's possible for a window to not have a viewport, in which case the viewport dimensions will be zero.
+			FVector2f ViewportCursorPosition = AppCursorPosition - Window->GetPositionInScreen();
+			FIntPoint ViewportOffset = FIntPoint::ZeroValue;
+			FIntPoint ViewportExtent = FIntPoint::ZeroValue;
+			FIntRect ViewportRect;
+			float ViewportScaleUI = Window->GetViewportScaleUIOverride();
 
-			if (ViewportWidget)
+			if (ISlateViewport* Viewport = Window->GetViewport().Get())
 			{
-				ViewportOffset = FIntPoint(
-					FMath::RoundToInt32(ViewportWidget->GetTickSpaceGeometry().GetAbsolutePosition().X - Window->GetPositionInScreen().X),
-					FMath::RoundToInt32(ViewportWidget->GetTickSpaceGeometry().GetAbsolutePosition().Y - Window->GetPositionInScreen().Y));
+				TSharedPtr<SWidget> ViewportWidget = Viewport->GetWidget().Pin();
 
-				ViewportCursorPosition -= ViewportWidget->GetPaintSpaceGeometry().AbsolutePosition;
-			}
+				if (ViewportWidget)
+				{
+					ViewportOffset = FIntPoint(
+						FMath::RoundToInt32(ViewportWidget->GetTickSpaceGeometry().GetAbsolutePosition().X - Window->GetPositionInScreen().X),
+						FMath::RoundToInt32(ViewportWidget->GetTickSpaceGeometry().GetAbsolutePosition().Y - Window->GetPositionInScreen().Y));
 
-			ViewportExtent = Viewport->GetSize();
-			ViewportRect = FIntRect(ViewportOffset, ViewportOffset + ViewportExtent);
+					ViewportCursorPosition -= ViewportWidget->GetPaintSpaceGeometry().AbsolutePosition;
+				}
 
-			ensureMsgf(AppViewportSceneFormat == PF_Unknown || AppViewportSceneFormat == Viewport->GetSceneTargetFormat(),
-				TEXT("Multiple viewport formats coming from multiple windows are not a supported scenario in slate. This will cause undefined behavior with Slate Post Buffers."));
+				ViewportExtent = Viewport->GetSize();
+				ViewportRect = FIntRect(ViewportOffset, ViewportOffset + ViewportExtent);
 
-			AppViewportSceneFormat = Viewport->GetSceneTargetFormat();
-			AppViewportExtentMax = AppViewportExtentMax.ComponentMax(ViewportExtent);
+				ensureMsgf(AppViewportSceneFormat == PF_Unknown || AppViewportSceneFormat == Viewport->GetSceneTargetFormat(),
+					TEXT("Multiple viewport formats coming from multiple windows are not a supported scenario in slate. This will cause undefined behavior with Slate Post Buffers."));
 
-			if (ViewportScaleUI < 0.0f)
-			{
-				ViewportScaleUI = GetDefault<UUserInterfaceSettings>()->GetDPIScaleBasedOnSize(ViewportExtent);
+				AppViewportSceneFormat = Viewport->GetSceneTargetFormat();
+				AppViewportExtentMax = AppViewportExtentMax.ComponentMax(ViewportExtent);
+
+				if (ViewportScaleUI < 0.0f)
+				{
+					ViewportScaleUI = GetDefault<UUserInterfaceSettings>()->GetDPIScaleBasedOnSize(ViewportExtent);
+				}
 			}
-		}
 
-		FSlateViewportInfo* ViewportInfo = WindowToViewportInfo.FindRef(Window);
+			FSlateViewportInfo* ViewportInfo = WindowToViewportInfo.FindRef(Window);
 
-		if (Window->IsViewportSizeDrivenByWindow())
-		{
-			ResizeViewportIfNeeded(ViewportInfo, ViewportInfo->ExtentToResizeTo, IsViewportFullscreen(*Window), Window);
-		}
+			if (Window->IsViewportSizeDrivenByWindow())
+			{
+				ResizeViewportIfNeeded(ViewportInfo, ViewportInfo->ExtentToResizeTo, IsViewportFullscreen(*Window), Window);
+			}
 
-		Window->SetIsHDR(ViewportInfo->bDisplayFormatIsHDR);
-		Window->ResetViewportScaleUIOverride();
+			Window->SetIsHDR(ViewportInfo->bDisplayFormatIsHDR);
+			Window->ResetViewportScaleUIOverride();
 
-		ElementBatcher->SetCompositeHDRViewports(ViewportInfo->bDisplayFormatIsHDR && CompositeUIWithSceneHDR());
-		ElementBatcher->AddElements(*WindowElementList);
+			ElementBatcher->SetCompositeHDRViewports(ViewportInfo->bDisplayFormatIsHDR && CompositeUIWithSceneHDR());
+			ElementBatcher->AddElements(*WindowElementList);
 
-		const bool bWindowCanRenderPostProcess       = bAppCanRender && ViewportExtent != FIntPoint::ZeroValue;
-		const ESlatePostRT PostProcessUsedBits       = bWindowCanRenderPostProcess ? ElementBatcher->GetUsedSlatePostBuffers() : ESlatePostRT::None;
-		const ESlatePostRT PostProcessCustomDrawBits = bWindowCanRenderPostProcess ? ElementBatcher->GetResourceUpdatingPostBuffers() : ESlatePostRT::None;
-		const ESlatePostRT PostProcessSkipUpdateBits = bWindowCanRenderPostProcess ? ElementBatcher->GetSkipDefaultUpdatePostBuffers() : ESlatePostRT::None;
-		const bool bLockToVsync                      = IsVSyncRequired(*ElementBatcher);
+			const bool bWindowCanRenderPostProcess       = bAppCanRender && ViewportExtent != FIntPoint::ZeroValue;
+			const ESlatePostRT PostProcessUsedBits       = bWindowCanRenderPostProcess ? ElementBatcher->GetUsedSlatePostBuffers() : ESlatePostRT::None;
+			const ESlatePostRT PostProcessCustomDrawBits = bWindowCanRenderPostProcess ? ElementBatcher->GetResourceUpdatingPostBuffers() : ESlatePostRT::None;
+			const ESlatePostRT PostProcessSkipUpdateBits = bWindowCanRenderPostProcess ? ElementBatcher->GetSkipDefaultUpdatePostBuffers() : ESlatePostRT::None;
+			const bool bLockToVsync                      = IsVSyncRequired(*ElementBatcher);
 
-		ElementBatcher->ResetBatches();
-		ElementBatcher->SetCompositeHDRViewports(false);
+			ElementBatcher->ResetBatches();
+			ElementBatcher->SetCompositeHDRViewports(false);
 
-		WindowsToRender.Emplace(FWindowToRender
-		{
-			  .Window                    = Window
-			, .WindowElementList         = WindowElementList
-			, .ViewportInfo              = ViewportInfo
-			, .ViewportOffset            = ViewportOffset
-			, .ViewportExtent            = ViewportExtent
-			, .ViewportRect              = ViewportRect
-			, .ViewportScaleUI           = ViewportScaleUI
-			, .PostProcessUsedBits       = PostProcessUsedBits
-			, .PostProcessCustumDrawBits = PostProcessCustomDrawBits
-			, .PostProcessSkipUpdateBits = PostProcessSkipUpdateBits
-			, .CursorPosition            = FIntPoint(ViewportCursorPosition.X, ViewportCursorPosition.Y)
-			, .bLockToVsync              = bLockToVsync
-		});
+			WindowsToRender.Emplace(FWindowToRender
+			{
+				  .Window                    = Window
+				, .WindowElementList         = WindowElementList
+				, .ViewportInfo              = ViewportInfo
+				, .ViewportOffset            = ViewportOffset
+				, .ViewportExtent            = ViewportExtent
+				, .ViewportRect              = ViewportRect
+				, .ViewportScaleUI           = ViewportScaleUI
+				, .PostProcessUsedBits       = PostProcessUsedBits
+				, .PostProcessCustumDrawBits = PostProcessCustomDrawBits
+				, .PostProcessSkipUpdateBits = PostProcessSkipUpdateBits
+				, .CursorPosition            = FIntPoint(ViewportCursorPosition.X, ViewportCursorPosition.Y)
+				, .bLockToVsync              = bLockToVsync
+			});
 
-		PostProcessAnyUsedBits |= PostProcessUsedBits;
+			PostProcessAnyUsedBits |= PostProcessUsedBits;
+		}
 	}
 
 	// Update the font cache now that all element batches were processed.
