Skip rendering Slate window if we don't have a valid viewport rather than checking if RenderOffscreen is enabled, as Vulkan ends up setting that flag while expecting to render UI.

diff --git a/Engine/Source/Runtime/SlateRHIRenderer/Private/SlateRHIRenderer.cpp b/Engine/Source/Runtime/SlateRHIRenderer/Private/SlateRHIRenderer.cpp
--- a/Engine/Source/Runtime/SlateRHIRenderer/Private/SlateRHIRenderer.cpp
+++ b/Engine/Source/Runtime/SlateRHIRenderer/Private/SlateRHIRenderer.cpp
@@ -1168,11 +1168,10 @@ void FSlateRHIRenderer::DrawWindows_Private(FSlateDrawBuffer& WindowDrawBuffer)
 
 	USlateRHIRendererSettings* RendererSettings = USlateRHIRendererSettings::GetMutable();
 
-	static const bool bRenderOffscreen = FParse::Param(FCommandLine::Get(), TEXT("RenderOffScreen"));
 	const float AppDeltaTime = FApp::GetDeltaTime();
 	const FGameTime AppDilatedTime = FGameTime::CreateDilated(FPlatformTime::Seconds() - GStartTime, AppDeltaTime, FApp::GetCurrentTime() - GStartTime, AppDeltaTime);
 	const FVector2f AppCursorPosition = FSlateApplication::Get().GetCursorPos();
-	const bool bAppCanRender = GIsClient && !IsRunningCommandlet() && !GUsingNullRHI && !bRenderOffscreen;
+	const bool bAppCanRender = GIsClient && !IsRunningCommandlet() && !GUsingNullRHI;
 	const bool bAppCanRenderPostProcess = RendererSettings && IsInGameThread() && UAssetManager::IsInitialized() && bAppCanRender && CVarCopyBackbufferToSlatePostRenderTargets.GetValueOnGameThread() > 0;
 	EPixelFormat AppViewportSceneFormat = PF_Unknown;
 	FIntPoint AppViewportExtentMax = FIntPoint::ZeroValue;
@@ -1257,45 +1256,46 @@ void FSlateRHIRenderer::DrawWindows_Private(FSlateDrawBuffer& WindowDrawBuffer)
 				}
 			}
 
-			FSlateViewportInfo* ViewportInfo = WindowToViewportInfo.FindRef(Window);
-
-			if (Window->IsViewportSizeDrivenByWindow())
+			if (FSlateViewportInfo* ViewportInfo = WindowToViewportInfo.FindRef(Window))
 			{
-				ResizeViewportIfNeeded(ViewportInfo, ViewportInfo->ExtentToResizeTo, IsViewportFullscreen(*Window), Window);
-			}
+				if (Window->IsViewportSizeDrivenByWindow())
+				{
+					ResizeViewportIfNeeded(ViewportInfo, ViewportInfo->ExtentToResizeTo, IsViewportFullscreen(*Window), Window);
+				}
 
-			Window->SetIsHDR(ViewportInfo->bDisplayFormatIsHDR);
-			Window->ResetViewportScaleUIOverride();
+				Window->SetIsHDR(ViewportInfo->bDisplayFormatIsHDR);
+				Window->ResetViewportScaleUIOverride();
 
-			ElementBatcher->SetCompositeHDRViewports(ViewportInfo->bDisplayFormatIsHDR && CompositeUIWithSceneHDR());
-			ElementBatcher->AddElements(*WindowElementList);
+				ElementBatcher->SetCompositeHDRViewports(ViewportInfo->bDisplayFormatIsHDR && CompositeUIWithSceneHDR());
+				ElementBatcher->AddElements(*WindowElementList);
 
-			const bool bWindowCanRenderPostProcess       = bAppCanRender && ViewportExtent != FIntPoint::ZeroValue;
-			const ESlatePostRT PostProcessUsedBits       = bWindowCanRenderPostProcess ? ElementBatcher->GetUsedSlatePostBuffers() : ESlatePostRT::None;
-			const ESlatePostRT PostProcessCustomDrawBits = bWindowCanRenderPostProcess ? ElementBatcher->GetResourceUpdatingPostBuffers() : ESlatePostRT::None;
-			const ESlatePostRT PostProcessSkipUpdateBits = bWindowCanRenderPostProcess ? ElementBatcher->GetSkipDefaultUpdatePostBuffers() : ESlatePostRT::None;
-			const bool bLockToVsync                      = IsVSyncRequired(*ElementBatcher);
+				const bool bWindowCanRenderPostProcess       = bAppCanRender && ViewportExtent != FIntPoint::ZeroValue;
+				const ESlatePostRT PostProcessUsedBits       = bWindowCanRenderPostProcess ? ElementBatcher->GetUsedSlatePostBuffers() : ESlatePostRT::None;
+				const ESlatePostRT PostProcessCustomDrawBits = bWindowCanRenderPostProcess ? ElementBatcher->GetResourceUpdatingPostBuffers() : ESlatePostRT::None;
+				const ESlatePostRT PostProcessSkipUpdateBits = bWindowCanRenderPostProcess ? ElementBatcher->GetSkipDefaultUpdatePostBuffers() : ESlatePostRT::None;
+				const bool bLockToVsync                      = IsVSyncRequired(*ElementBatcher);
 
-			ElementBatcher->ResetBatches();
-			ElementBatcher->SetCompositeHDRViewports(false);
+				ElementBatcher->ResetBatches();
+				ElementBatcher->SetCompositeHDRViewports(false);
 
-			WindowsToRender.Emplace(FWindowToRender
-			{
-				  .Window                    = Window
-				, .WindowElementList         = WindowElementList
-				, .ViewportInfo              = ViewportInfo
-				, .ViewportOffset            = ViewportOffset
-				, .ViewportExtent            = ViewportExtent
-				, .ViewportRect              = ViewportRect
-				, .ViewportScaleUI           = ViewportScaleUI
-				, .PostProcessUsedBits       = PostProcessUsedBits
-				, .PostProcessCustumDrawBits = PostProcessCustomDrawBits
-				, .PostProcessSkipUpdateBits = PostProcessSkipUpdateBits
-				, .CursorPosition            = FIntPoint(ViewportCursorPosition.X, ViewportCursorPosition.Y)
-				, .bLockToVsync              = bLockToVsync
-			});
+				WindowsToRender.Emplace(FWindowToRender
+				{
+					  .Window                    = Window
+					, .WindowElementList         = WindowElementList
+					, .ViewportInfo              = ViewportInfo
+					, .ViewportOffset            = ViewportOffset
+					, .ViewportExtent            = ViewportExtent
+					, .ViewportRect              = ViewportRect
+					, .ViewportScaleUI           = ViewportScaleUI
+					, .PostProcessUsedBits       = PostProcessUsedBits
+					, .PostProcessCustumDrawBits = PostProcessCustomDrawBits
+					, .PostProcessSkipUpdateBits = PostProcessSkipUpdateBits
+					, .CursorPosition            = FIntPoint(ViewportCursorPosition.X, ViewportCursorPosition.Y)
+					, .bLockToVsync              = bLockToVsync
+				});
 
-			PostProcessAnyUsedBits |= PostProcessUsedBits;
+				PostProcessAnyUsedBits |= PostProcessUsedBits;
+			}
 		}
 	}
 
