Fixed D3D12 viewport/swapchain release with pending references
- Manually release async deleters

diff --git a/Engine/Source/Runtime/SlateRHIRenderer/Private/SlateRHIRenderer.cpp b/Engine/Source/Runtime/SlateRHIRenderer/Private/SlateRHIRenderer.cpp
--- a/Engine/Source/Runtime/SlateRHIRenderer/Private/SlateRHIRenderer.cpp
+++ b/Engine/Source/Runtime/SlateRHIRenderer/Private/SlateRHIRenderer.cpp
@@ -402,6 +402,12 @@ void FSlateRHIRenderer::ResizeViewportIfNeeded(FSlateViewportInfo* ViewInfo, FIn
 		//		this should be removed and whatever streaming thread safety is needed during a viewport resize should be done correctly
 		SuspendTextureStreamingRenderTasks();
 
+		// Wait for any pending async cleanup
+		ENQUEUE_RENDER_COMMAND(FAsyncCleanup)([](FRHICommandListImmediate&)
+		{
+			FRDGBuilder::WaitForAsyncDeleteTask();
+		});
+
 		// cannot resize the viewport while potentially using it.
 		FlushRenderingCommands();
 
