Fix clipping of controls under RetainerBox when offset is non zero

#tests Tested with grid cliptobounds with images inside under a RetainerBox widget with/without the fix. Also tried panning around. With the old code, panning would eventually clip images that were further on the surface
[RN] fix creator/ui

diff --git a/Engine/Source/Runtime/SlateRHIRenderer/Private/SlateRHIRenderingPolicy.cpp b/Engine/Source/Runtime/SlateRHIRenderer/Private/SlateRHIRenderingPolicy.cpp
--- a/Engine/Source/Runtime/SlateRHIRenderer/Private/SlateRHIRenderingPolicy.cpp
+++ b/Engine/Source/Runtime/SlateRHIRenderer/Private/SlateRHIRenderingPolicy.cpp
@@ -753,8 +753,8 @@ void SetSlateClipping(FRHICommandList& RHICmdList, const FSlateClippingOp* Clipp
 		{
 			ScissorRect.Left   = FMath::Clamp(ScissorRect.Left + ElementOffset.X, (float)ViewportRect.Min.X, (float)ViewportRect.Max.X);
 			ScissorRect.Top    = FMath::Clamp(ScissorRect.Top  + ElementOffset.Y, (float)ViewportRect.Min.Y, (float)ViewportRect.Max.Y);
-			ScissorRect.Right  = FMath::Clamp(ScissorRect.Right,  ScissorRect.Left, (float)ViewportRect.Max.X);
-			ScissorRect.Bottom = FMath::Clamp(ScissorRect.Bottom, ScissorRect.Top,  (float)ViewportRect.Max.Y);
+			ScissorRect.Right  = FMath::Clamp(ScissorRect.Right + ElementOffset.X,  ScissorRect.Left, (float)ViewportRect.Max.X);
+			ScissorRect.Bottom = FMath::Clamp(ScissorRect.Bottom + ElementOffset.Y, ScissorRect.Top,  (float)ViewportRect.Max.Y);
 			return ScissorRect;
 		};
 
