Add virtual texture feedback bindings to slate element pixel shader.

diff --git a/Engine/Shaders/Private/SlateElementPixelShader.usf b/Engine/Shaders/Private/SlateElementPixelShader.usf
--- a/Engine/Shaders/Private/SlateElementPixelShader.usf
+++ b/Engine/Shaders/Private/SlateElementPixelShader.usf
@@ -33,14 +33,6 @@ half4 DrawFlags;
 Texture2D ElementTexture;
 SamplerState ElementTextureSampler;
 
-#if SAMPLE_VIRTUAL_TEXTURE
-Texture2D<uint4> InPageTableTexture;
-uint4 VTPackedPageTableUniform0;
-uint4 VTPackedPageTableUniform1;
-uint4 VTPackedUniform;
-float4 VTShaderParams;
-#endif // SAMPLE_VIRTUAL_TEXTURE
-
 /**
  * Applys a contrast about a midpoint squeezing the colors towards dark or light depending on
  * the midpoint.  Additionally - it increases the opaqueness of colors to try and seperate them from
@@ -214,6 +206,10 @@ half4 GetMaterialColor( VertexToPixelInterpolants InVertex, float2 MeshUV, float
 
 #endif // SUBSTRATE_ENABLED
 
+#if MATERIAL_VIRTUALTEXTURE_FEEDBACK
+	FinalizeVirtualTextureFeedback(Parameters.VirtualTextureFeedback, Parameters.SvPosition, View.VTFeedbackBuffer);
+#endif
+
 	return OutColor;
 }
 #endif
@@ -224,26 +220,28 @@ half4 GetTextureColor( VertexToPixelInterpolants InVertex, float2 UV )
 	half4 BaseColor;
 
 	#if SAMPLE_VIRTUAL_TEXTURE
-			VTPageTableResult PageTableResult = TextureLoadVirtualPageTableLevel(InPageTableTexture, VTPageTableUniform_Unpack(VTPackedPageTableUniform0, VTPackedPageTableUniform1), UV, VTADDRESSMODE_WRAP, VTADDRESSMODE_WRAP, VTShaderParams.x); 
-			BaseColor = TextureVirtualSampleLevel(ElementTexture, ElementTextureSampler, PageTableResult, (uint)VTShaderParams.y, VTUniform_Unpack(VTPackedUniform));
+		FVirtualTextureFeedbackParams VirtualTextureFeedback;
+		InitializeVirtualTextureFeedback(VirtualTextureFeedback);
+		VTPageTableUniform PageTableUniform = VTPageTableUniform_Unpack(SlateElementVirtualTextureParams.PackedPageTableUniform0, SlateElementVirtualTextureParams.PackedPageTableUniform1);
+		VTPageTableResult PageTableResult = TextureLoadVirtualPageTable(SlateElementVirtualTextureParams.PageTableTexture, PageTableUniform, UV, VTADDRESSMODE_WRAP, VTADDRESSMODE_WRAP, 0.0f, InVertex.Position.xy, 0, VirtualTextureFeedback); 
+		BaseColor = TextureVirtualSampleLevel(ElementTexture, ElementTextureSampler, PageTableResult, SlateElementVirtualTextureParams.LayerIndex, VTUniform_Unpack(SlateElementVirtualTextureParams.PackedUniform));
+		FinalizeVirtualTextureFeedback(VirtualTextureFeedback, SlateElementVirtualTextureParams.FeedbackParams, InVertex.Position, SlateElementVirtualTextureParams.FeedbackBuffer);
 	#else
-
 		BaseColor = Texture2DSample(ElementTexture, ElementTextureSampler, UV );
-
-		#if USE_TEXTURE_GRAYSCALE
-			#if !USE_TEXTURE_ALPHA
-				BaseColor.rgb = BaseColor.rrr;
-				BaseColor.a = 1;
-			#else
-				BaseColor.a = BaseColor.r;
-				BaseColor.rgb = half3(1,1,1);
-			#endif
-		#endif
-
 	#endif
 
-	#if !USE_TEXTURE_ALPHA
-		BaseColor.a = 1.0f;
+	#if USE_TEXTURE_GRAYSCALE
+		#if !USE_TEXTURE_ALPHA
+			BaseColor.rgb = BaseColor.rrr;
+			BaseColor.a = 1;
+		#else
+			BaseColor.a = BaseColor.r;
+			BaseColor.rgb = half3(1,1,1);
+		#endif
+	#else
+		#if !USE_TEXTURE_ALPHA
+			BaseColor.a = 1.0f;
+		#endif
 	#endif
 
 	return BaseColor;
diff --git a/Engine/Shaders/Private/VirtualTextureCommon.ush b/Engine/Shaders/Private/VirtualTextureCommon.ush
--- a/Engine/Shaders/Private/VirtualTextureCommon.ush
+++ b/Engine/Shaders/Private/VirtualTextureCommon.ush
@@ -53,25 +53,31 @@ void StoreVirtualTextureFeedback(in out FVirtualTextureFeedbackParams Params, ui
 }
 
 /** This should be called at the end of the pixel shader to write out the gathered VT feedback info to the OutputBuffer. */
-void FinalizeVirtualTextureFeedback(in FVirtualTextureFeedbackParams Params, float4 SvPosition, RWStructuredBuffer<uint> OutputBuffer)
+void FinalizeVirtualTextureFeedback(in FVirtualTextureFeedbackParams Params, in uint4 FeedbackViewParams, in float4 SvPosition, RWStructuredBuffer<uint> OutputBuffer)
 {
-	uint2 PixelTilePos = (uint2)SvPosition.xy & View.VirtualTextureFeedbackMask;
-	uint PixelTileIndex = (PixelTilePos.y << View.VirtualTextureFeedbackShift) + PixelTilePos.x;
+	uint2 PixelTilePos = (uint2)SvPosition.xy & FeedbackViewParams.x;
+	uint PixelTileIndex = (PixelTilePos.y << FeedbackViewParams.y) + PixelTilePos.x;
 
 	// This code will only run every few pixels...
-	[branch] if (PixelTileIndex == View.VirtualTextureFeedbackJitterOffset)
+	[branch] if (PixelTileIndex == FeedbackViewParams.z)
 	{
 		uint WriteOffset = 0;
 		InterlockedAdd(OutputBuffer[0], 1, WriteOffset);
 		WriteOffset += 1;
 
-		if (WriteOffset < View.VirtualTextureFeedbackBufferSize)
+		if (WriteOffset < FeedbackViewParams.w)
 		{
 			OutputBuffer[WriteOffset] = Params.Request;
 		}
 	}	
 }
 
+void FinalizeVirtualTextureFeedback(in FVirtualTextureFeedbackParams Params, float4 SvPosition, RWStructuredBuffer<uint> OutputBuffer)
+{
+	const uint4 FeedbackViewParams = uint4(View.VirtualTextureFeedbackMask, View.VirtualTextureFeedbackShift, View.VirtualTextureFeedbackJitterOffset, View.VirtualTextureFeedbackBufferSize);
+	FinalizeVirtualTextureFeedback(Params, FeedbackViewParams, SvPosition, OutputBuffer);
+}
+
 /**
  * Address Mode Logic
  * Apply before after calculating mip level/derivatives!
diff --git a/Engine/Source/Runtime/SlateRHIRenderer/Private/SlateRHIRenderingPolicy.cpp b/Engine/Source/Runtime/SlateRHIRenderer/Private/SlateRHIRenderingPolicy.cpp
--- a/Engine/Source/Runtime/SlateRHIRenderer/Private/SlateRHIRenderingPolicy.cpp
+++ b/Engine/Source/Runtime/SlateRHIRenderer/Private/SlateRHIRenderingPolicy.cpp
@@ -35,6 +35,7 @@
 #include "Containers/StaticBitArray.h"
 #include "MeshPassProcessor.h"
 #include "PSOPrecacheValidation.h"
+#include "VT/VirtualTextureFeedbackResource.h"
 
 extern void UpdateNoiseTextureParameters(FViewUniformShaderParameters& ViewUniformShaderParameters);
 
@@ -675,7 +676,6 @@ private:
 		};
 
 		FViewUniformShaderParameters ViewUniformShaderParameters;
-		ViewUniformShaderParameters.VTFeedbackBuffer = GEmptyStructuredBufferWithUAV->UnorderedAccessViewRHI;
 
 		SetupCommonViewUniformBufferParameters(ViewUniformShaderParameters, Inputs.TextureExtent, 1, ViewRect, ViewMatrices, ViewMatrices, SetupViewUniformParameterInputs);
 
@@ -691,6 +691,11 @@ private:
 		ViewUniformShaderParameters.MobilePreviewMode = (FeatureLevel == ERHIFeatureLevel::ES3_1) && GMaxRHIFeatureLevel > ERHIFeatureLevel::ES3_1 ? 1.0f : 0.0f;
 
 		UpdateNoiseTextureParameters(ViewUniformShaderParameters);
+
+		VirtualTexture::FFeedbackShaderParams VirtualTextureFeedbackShaderParams;
+		VirtualTexture::GetFeedbackShaderParams(VirtualTextureFeedbackShaderParams);
+		VirtualTexture::UpdateViewUniformShaderParameters(VirtualTextureFeedbackShaderParams, ViewUniformShaderParameters);
+
 		return TUniformBufferRef<FViewUniformShaderParameters>::CreateUniformBufferImmediate(ViewUniformShaderParameters, UniformBuffer_SingleFrame);
 	}
 
diff --git a/Engine/Source/Runtime/SlateRHIRenderer/Private/SlateShaders.cpp b/Engine/Source/Runtime/SlateRHIRenderer/Private/SlateShaders.cpp
--- a/Engine/Source/Runtime/SlateRHIRenderer/Private/SlateShaders.cpp
+++ b/Engine/Source/Runtime/SlateRHIRenderer/Private/SlateShaders.cpp
@@ -1,9 +1,10 @@
 // Copyright Epic Games, Inc. All Rights Reserved.
 
 #include "SlateShaders.h"
-#include "Rendering/RenderingCommon.h"
-#include "PipelineStateCache.h"
 #include "DataDrivenShaderPlatformInfo.h"
+#include "PipelineStateCache.h"
+#include "Rendering/RenderingCommon.h"
+#include "VT/VirtualTextureFeedbackResource.h"
 
 /** Flag to determine if we are running with a color vision deficiency shader on */
 EColorVisionDeficiency GSlateColorDeficiencyType = EColorVisionDeficiency::NormalVision;
@@ -138,4 +139,60 @@ void FSlateMaskingVertexDeclaration::InitRHI(FRHICommandListBase& RHICmdList)
 void FSlateMaskingVertexDeclaration::ReleaseRHI()
 {
 	VertexDeclarationRHI.SafeRelease();
+}
+
+/** Uniform buffer layout for sampling virtual texture. */
+BEGIN_GLOBAL_SHADER_PARAMETER_STRUCT(FSlateElementVirtualTextureParams, )
+	SHADER_PARAMETER_TEXTURE(Texture2D<uint4>, PageTableTexture)
+	SHADER_PARAMETER_UAV(RWStructuredBuffer<uint>, FeedbackBuffer)
+	SHADER_PARAMETER(FUintVector4, PackedPageTableUniform0)
+	SHADER_PARAMETER(FUintVector4, PackedPageTableUniform1)
+	SHADER_PARAMETER(FUintVector4, PackedUniform)
+	SHADER_PARAMETER(FUintVector4, FeedbackParams)
+	SHADER_PARAMETER(uint32, LayerIndex)
+END_GLOBAL_SHADER_PARAMETER_STRUCT()
+
+IMPLEMENT_GLOBAL_SHADER_PARAMETER_STRUCT(FSlateElementVirtualTextureParams, "SlateElementVirtualTextureParams");
+
+void FSlateElementPS::SetVirtualTextureParameters(FMeshDrawSingleShaderBindings& ShaderBindings, FVirtualTexture2DResource* InVirtualTexture)
+{
+	if (InVirtualTexture == nullptr)
+	{
+		return;
+	}
+
+	const uint32 LayerIndex = 0;
+
+	IAllocatedVirtualTexture* AllocatedVT = InVirtualTexture->AcquireAllocatedVT();
+	FRHIShaderResourceView* PhysicalView = AllocatedVT->GetPhysicalTextureSRV(LayerIndex, InVirtualTexture->bSRGB);
+
+	ShaderBindings.Add(TextureParameter, PhysicalView);
+	ShaderBindings.Add(TextureParameterSampler, InVirtualTexture->SamplerStateRHI);
+
+	FUintVector4 PackedPageTableUniform[2];
+	AllocatedVT->GetPackedPageTableUniform(PackedPageTableUniform);
+
+	FUintVector4 PackedUniform;
+	AllocatedVT->GetPackedUniform(&PackedUniform, LayerIndex);
+
+	VirtualTexture::FFeedbackShaderParams VirtualTextureFeedbackShaderParams;
+	VirtualTexture::GetFeedbackShaderParams(VirtualTextureFeedbackShaderParams);
+
+	const FUintVector4 FeedbackParams{ 
+		VirtualTextureFeedbackShaderParams.Mask, 
+		VirtualTextureFeedbackShaderParams.Shift, 
+		VirtualTextureFeedbackShaderParams.JitterOffset, 
+		VirtualTextureFeedbackShaderParams.BufferSize };
+
+	FSlateElementVirtualTextureParams VTParams;
+	VTParams.PageTableTexture = AllocatedVT->GetPageTableTexture(LayerIndex);
+	VTParams.FeedbackBuffer = VirtualTextureFeedbackShaderParams.BufferUAV;
+	VTParams.PackedPageTableUniform0 = PackedPageTableUniform[0];
+	VTParams.PackedPageTableUniform1 = PackedPageTableUniform[1];
+	VTParams.PackedUniform = PackedUniform;
+	VTParams.FeedbackParams = FeedbackParams;
+	VTParams.LayerIndex = LayerIndex;
+
+	TUniformBufferRef<FSlateElementVirtualTextureParams> VTParamsUB = CreateUniformBufferImmediate(VTParams, EUniformBufferUsage::UniformBuffer_SingleFrame);
+	ShaderBindings.Add(VirtualTextureParams, VTParamsUB);
 }
\ No newline at end of file
diff --git a/Engine/Source/Runtime/SlateRHIRenderer/Private/SlateShaders.h b/Engine/Source/Runtime/SlateRHIRenderer/Private/SlateShaders.h
--- a/Engine/Source/Runtime/SlateRHIRenderer/Private/SlateShaders.h
+++ b/Engine/Source/Runtime/SlateRHIRenderer/Private/SlateShaders.h
@@ -6,15 +6,12 @@
 #include "CoreMinimal.h"
 #include "RenderResource.h"
 #include "ShaderParameters.h"
-#include "Shader.h"
+#include "GlobalRenderResources.h"
 #include "GlobalShader.h"
-#include "ShaderParameterUtils.h"
+#include "MeshDrawShaderBindings.h"
 #include "Rendering/RenderingCommon.h"
-#include "RHIStaticStates.h"
-#include "TextureResource.h"
-#include "RenderUtils.h"
 #include "ShaderParameterStruct.h"
-#include "MeshDrawShaderBindings.h"
+#include "TextureResource.h"
 
 extern EColorVisionDeficiency GSlateColorDeficiencyType;
 extern int32 GSlateColorDeficiencySeverity;
@@ -88,14 +85,10 @@ public:
 	{
 		TextureParameter.Bind(Initializer.ParameterMap, TEXT("ElementTexture"));
 		TextureParameterSampler.Bind(Initializer.ParameterMap, TEXT("ElementTextureSampler"));
-		InPageTableTexture.Bind(Initializer.ParameterMap, TEXT("InPageTableTexture"));
-		VTPackedPageTableUniform0.Bind(Initializer.ParameterMap, TEXT("VTPackedPageTableUniform0"));
-		VTPackedPageTableUniform1.Bind(Initializer.ParameterMap, TEXT("VTPackedPageTableUniform1"));
-		VTPackedUniform.Bind(Initializer.ParameterMap, TEXT("VTPackedUniform"));
 		ShaderParams.Bind(Initializer.ParameterMap, TEXT("ShaderParams"));
 		ShaderParams2.Bind(Initializer.ParameterMap, TEXT("ShaderParams2"));
-		VTShaderParams.Bind(Initializer.ParameterMap, TEXT("VTShaderParams"));
 		GammaAndAlphaValues.Bind(Initializer.ParameterMap,TEXT("GammaAndAlphaValues"));
+		VirtualTextureParams.Bind(Initializer.ParameterMap, TEXT("SlateElementVirtualTextureParams"));
 	}
 
 	static void ModifyCompilationEnvironment(const FGlobalShaderPermutationParameters& Parameters, FShaderCompilerEnvironment& OutEnvironment);
@@ -116,34 +109,7 @@ public:
 	 *
 	 * @param InVirtualTexture	Virtual Texture resource to use when this pixel shader is bound
 	 */
-	void SetVirtualTextureParameters(FMeshDrawSingleShaderBindings& ShaderBindings, FVirtualTexture2DResource* InVirtualTexture)
-	{
-		if (InVirtualTexture == nullptr)
-		{
-			return;
-		}
-
-		IAllocatedVirtualTexture* AllocatedVT = InVirtualTexture->AcquireAllocatedVT();
-		uint32 LayerIndex = 0;
-
-		FRHIShaderResourceView* PhysicalView = AllocatedVT->GetPhysicalTextureSRV(LayerIndex, InVirtualTexture->bSRGB);
-
-		FUintVector4 PageTableUniform[2];
-		FUintVector4 Uniform;
-		// VTParams.X = MipLevel, VTParams.Y = LayerIndex
-		FVector4f VTParams{ 0.f, static_cast<float>(LayerIndex), 0.f, 0.f };
-
-		AllocatedVT->GetPackedPageTableUniform(PageTableUniform);
-		AllocatedVT->GetPackedUniform(&Uniform, LayerIndex);
-
-		ShaderBindings.Add(TextureParameter, PhysicalView);
-		ShaderBindings.Add(TextureParameterSampler, InVirtualTexture->SamplerStateRHI);
-		ShaderBindings.Add(InPageTableTexture, AllocatedVT->GetPageTableTexture(0u));
-		ShaderBindings.Add(VTPackedPageTableUniform0, PageTableUniform[0]);
-		ShaderBindings.Add(VTPackedPageTableUniform1, PageTableUniform[1]);
-		ShaderBindings.Add(VTPackedUniform, Uniform);
-		ShaderBindings.Add(VTShaderParams, VTParams);
-	}
+	void SetVirtualTextureParameters(FMeshDrawSingleShaderBindings& ShaderBindings, FVirtualTexture2DResource* InVirtualTexture);
 
 	/**
 	 * Sets shader params used by the shader
@@ -171,14 +137,10 @@ public:
 private:
 	LAYOUT_FIELD(FShaderResourceParameter, TextureParameter);
 	LAYOUT_FIELD(FShaderResourceParameter, TextureParameterSampler);
-	LAYOUT_FIELD(FShaderResourceParameter, InPageTableTexture);
-	LAYOUT_FIELD(FShaderParameter, VTPackedPageTableUniform0);
-	LAYOUT_FIELD(FShaderParameter, VTPackedPageTableUniform1);
-	LAYOUT_FIELD(FShaderParameter, VTPackedUniform);
 	LAYOUT_FIELD(FShaderParameter, ShaderParams);
 	LAYOUT_FIELD(FShaderParameter, ShaderParams2);
-	LAYOUT_FIELD(FShaderParameter, VTShaderParams);
 	LAYOUT_FIELD(FShaderParameter, GammaAndAlphaValues);
+	LAYOUT_FIELD(FShaderUniformBufferParameter, VirtualTextureParams);
 };
 
 /** 
