Fix editor only issue where slate caches the incorrect virtual texture state for a slate texture object because it was compiling at the time it was cached.
We no longer cache the result in the editor path.

diff --git a/Engine/Source/Runtime/SlateRHIRenderer/Private/SlateUTextureResource.cpp b/Engine/Source/Runtime/SlateRHIRenderer/Private/SlateUTextureResource.cpp
--- a/Engine/Source/Runtime/SlateRHIRenderer/Private/SlateUTextureResource.cpp
+++ b/Engine/Source/Runtime/SlateRHIRenderer/Private/SlateUTextureResource.cpp
@@ -13,7 +13,7 @@ FSlateBaseUTextureResource::FSlateBaseUTextureResource(UTexture* InTexture)
 	UpdateDebugName();
 #endif
 
-	if (TextureObject&& TextureObject->IsCurrentlyVirtualTextured())
+	if (TextureObject != nullptr && TextureObject->IsCurrentlyVirtualTextured())
 	{
 		CachedIsVirtualTexture = true;
 	}
@@ -45,7 +45,12 @@ ESlatePostRT FSlateBaseUTextureResource::GetUsedSlatePostBuffers() const
 
 bool FSlateBaseUTextureResource::RequiresVirtualTextureFeedback() const
 {
+#if WITH_EDITOR
+	// In editor we can't cache the state because it changes when texture is compiling.
+	return TextureObject && TextureObject->IsCurrentlyVirtualTextured();
+#else
 	return CachedIsVirtualTexture;
+#endif
 }
 
 #if SLATE_CHECK_UOBJECT_RENDER_RESOURCES
