Temporary removal of editor per update check for whether a slate texture resource is a virtual texture. This is triggering an ensure().

diff --git a/Engine/Source/Runtime/SlateRHIRenderer/Private/SlateUTextureResource.cpp b/Engine/Source/Runtime/SlateRHIRenderer/Private/SlateUTextureResource.cpp
--- a/Engine/Source/Runtime/SlateRHIRenderer/Private/SlateUTextureResource.cpp
+++ b/Engine/Source/Runtime/SlateRHIRenderer/Private/SlateUTextureResource.cpp
@@ -45,7 +45,7 @@ ESlatePostRT FSlateBaseUTextureResource::GetUsedSlatePostBuffers() const
 
 bool FSlateBaseUTextureResource::RequiresVirtualTextureFeedback() const
 {
-#if WITH_EDITOR
+#if 0//WITH_EDITOR
 	// In editor we can't cache the state because it changes when texture is compiling.
 	return TextureObject && TextureObject->IsCurrentlyVirtualTextured();
 #else
