Add back editor check for whether a slate texture resource is a virtual texture. Done once at init, but with a less conservative check that works even when texture is compiling.
Doing the check per frame inside RequiresVirtualTextureFeedback() could have race condition issues around thread checks when retrieving the resource, because it is called on a task spawned by game thread and waited on (and possibly retracted) by the render thread.

diff --git a/Engine/Source/Runtime/SlateRHIRenderer/Private/SlateUTextureResource.cpp b/Engine/Source/Runtime/SlateRHIRenderer/Private/SlateUTextureResource.cpp
--- a/Engine/Source/Runtime/SlateRHIRenderer/Private/SlateUTextureResource.cpp
+++ b/Engine/Source/Runtime/SlateRHIRenderer/Private/SlateUTextureResource.cpp
@@ -1,6 +1,8 @@
 // Copyright Epic Games, Inc. All Rights Reserved.
 
 #include "SlateUTextureResource.h"
+
+#include "RenderUtils.h"
 #include "SlateRHIRendererSettings.h"
 
 FSlateBaseUTextureResource::FSlateBaseUTextureResource(UTexture* InTexture)
@@ -13,10 +15,12 @@ FSlateBaseUTextureResource::FSlateBaseUTextureResource(UTexture* InTexture)
 	UpdateDebugName();
 #endif
 
-	if (TextureObject != nullptr && TextureObject->IsCurrentlyVirtualTextured())
-	{
-		CachedIsVirtualTexture = true;
-	}
+#if WITH_EDITOR
+	// Use a looser condition when in editor because the texture may still be compiling which would return false for IsCurrentlyVirtualTextured()
+	CachedIsVirtualTexture = TextureObject != nullptr && TextureObject->VirtualTextureStreaming && UseVirtualTexturing(GMaxRHIShaderPlatform);
+#else
+	CachedIsVirtualTexture = TextureObject != nullptr && TextureObject->IsCurrentlyVirtualTextured();
+#endif
 }
 
 FSlateBaseUTextureResource::~FSlateBaseUTextureResource()
@@ -45,12 +49,7 @@ ESlatePostRT FSlateBaseUTextureResource::GetUsedSlatePostBuffers() const
 
 bool FSlateBaseUTextureResource::RequiresVirtualTextureFeedback() const
 {
-#if 0//WITH_EDITOR
-	// In editor we can't cache the state because it changes when texture is compiling.
-	return TextureObject && TextureObject->IsCurrentlyVirtualTextured();
-#else
 	return CachedIsVirtualTexture;
-#endif
 }
 
 #if SLATE_CHECK_UOBJECT_RENDER_RESOURCES
