Fixed offset when using a background blur widget within a retainer panel
See UDN: https://epicgames.lightning.force.com/lightning/r/Case/500QP00000aKQbAYAW/view

[RN] Fixed offset when using a background blur widget within a retainer panel

diff --git a/Engine/Source/Runtime/SlateRHIRenderer/Private/SlateRHIRenderingPolicy.cpp b/Engine/Source/Runtime/SlateRHIRenderer/Private/SlateRHIRenderingPolicy.cpp
--- a/Engine/Source/Runtime/SlateRHIRenderer/Private/SlateRHIRenderingPolicy.cpp
+++ b/Engine/Source/Runtime/SlateRHIRenderer/Private/SlateRHIRenderingPolicy.cpp
@@ -1643,7 +1643,7 @@ void AddSlateDrawElementsPass(
 				BlurInputs.InputTexture = Inputs.ElementsTexture;
 			}
 
-			BlurInputs.InputRect = FIntRect(ShaderParams.PixelParams.X, ShaderParams.PixelParams.Y, ShaderParams.PixelParams.Z, ShaderParams.PixelParams.W);
+			BlurInputs.InputRect = FIntRect(ShaderParams.PixelParams.X + Inputs.ElementsOffset.X, ShaderParams.PixelParams.Y + Inputs.ElementsOffset.Y, ShaderParams.PixelParams.Z + Inputs.ElementsOffset.X, ShaderParams.PixelParams.W + Inputs.ElementsOffset.Y);
 			BlurInputs.OutputTexture = Inputs.SceneViewportTexture ? Inputs.SceneViewportTexture : Inputs.ElementsTexture;
 			BlurInputs.OutputRect = BlurInputs.InputRect;
 			BlurInputs.OutputLoadAction = ConsumeLoadAction(ElementsLoadAction);
