Fixed artifacts when HDR is enabled and a gaussian blur is used as the first element. The post process item was erroneously consuming the load action meant for the ElementsTexture. Also, the load action for clearing the rectangle for upsampling should always be Load, because it is meant to clear out existing contents.

diff --git a/Engine/Source/Runtime/SlateRHIRenderer/Private/SlatePostProcessor.cpp b/Engine/Source/Runtime/SlateRHIRenderer/Private/SlatePostProcessor.cpp
--- a/Engine/Source/Runtime/SlateRHIRenderer/Private/SlatePostProcessor.cpp
+++ b/Engine/Source/Runtime/SlateRHIRenderer/Private/SlatePostProcessor.cpp
@@ -230,7 +230,7 @@ void AddSlatePostProcessUpsamplePass(FRDGBuilder& GraphBuilder, const FSlatePost
 			? ESlatePostProcessUpsampleOutputFormat::HDR_SCRGB
 			: ESlatePostProcessUpsampleOutputFormat::HDR_PQ10;
 
-		PassParameters->RenderTargets[1] = FRenderTargetBinding(Inputs.OutputTextureToClear, Inputs.OutputLoadAction);
+		PassParameters->RenderTargets[1] = FRenderTargetBinding(Inputs.OutputTextureToClear, ERenderTargetLoadAction::ELoad);
 	}
 
 	FSlatePostProcessUpsamplePS::FPermutationDomain PermutationVector;
