Disable explicit color deficiency correction pass when HDR composition is enabled, as we perform it during the composition pass instead.

diff --git a/Engine/Source/Runtime/SlateRHIRenderer/Private/SlateRHIRenderer.cpp b/Engine/Source/Runtime/SlateRHIRenderer/Private/SlateRHIRenderer.cpp
--- a/Engine/Source/Runtime/SlateRHIRenderer/Private/SlateRHIRenderer.cpp
+++ b/Engine/Source/Runtime/SlateRHIRenderer/Private/SlateRHIRenderer.cpp
@@ -802,19 +802,20 @@ FSlateDrawWindowPassOutputs FSlateRHIRenderer::DrawWindow_RenderThread(FRDGBuild
 
 		if (bCompositeUIWithSceneHDR)
 		{
+			// Color deficiency correction is performed inside of the CompositeUI pass instead.
+			DrawElementsInputs.bAllowColorDeficiencyCorrection = false;
+
 			if (!BatchDataHDR.GetRenderBatches().IsEmpty())
 			{
 				DrawElementsInputs.ElementsTexture = OutputTexture;
 				DrawElementsInputs.ElementsLoadAction = ERenderTargetLoadAction::EClear;
 				DrawElementsInputs.ElementsBuffers = SlateElementsBuffersHDR;
 				DrawElementsInputs.StencilTexture = BatchDataHDR.IsStencilClippingRequired() ? SlateStencilTexture : nullptr;
-				DrawElementsInputs.bAllowColorDeficiencyCorrection = false;
 
 				AddSlateDrawElementsPass(GraphBuilder, *RenderingPolicy, DrawElementsInputs, BatchDataHDR.GetRenderBatches(), BatchDataHDR.GetFirstRenderBatchIndex());
 			}
 
 			DrawElementsInputs.bElementsTextureIsHDRDisplay = false;
-			DrawElementsInputs.bAllowColorDeficiencyCorrection = true;
 		}
 
 		DrawElementsInputs.ElementsTexture = ElementsTexture;
