Always defer slate widget rendering when possible. Absorb the widgets into the main scene render if possible to avoid frame-late issues for UWidgetComponent. This fixes performance issues when many widgets are rendered, which is now worse after the RDG refactor.

diff --git a/Engine/Source/Runtime/Engine/Private/SceneRenderBuilderInterface.cpp b/Engine/Source/Runtime/Engine/Private/SceneRenderBuilderInterface.cpp
--- a/Engine/Source/Runtime/Engine/Private/SceneRenderBuilderInterface.cpp
+++ b/Engine/Source/Runtime/Engine/Private/SceneRenderBuilderInterface.cpp
@@ -7,4 +7,10 @@
 TUniquePtr<ISceneRenderBuilder> ISceneRenderBuilder::Create(FSceneInterface* SceneInterface)
 {
 	return GetRendererModule().CreateSceneRenderBuilder(SceneInterface);
+}
+
+ISceneRenderBuilder::FOnPreSceneRender& ISceneRenderBuilder::OnPreSceneRender()
+{
+	static ISceneRenderBuilder::FOnPreSceneRender OnSceneRenderDelegate;
+	return OnSceneRenderDelegate;
 }
\ No newline at end of file
diff --git a/Engine/Source/Runtime/Engine/Public/SceneRenderBuilderInterface.h b/Engine/Source/Runtime/Engine/Public/SceneRenderBuilderInterface.h
--- a/Engine/Source/Runtime/Engine/Public/SceneRenderBuilderInterface.h
+++ b/Engine/Source/Runtime/Engine/Public/SceneRenderBuilderInterface.h
@@ -101,6 +101,9 @@ public:
 	// Returns a linear allocator for allocating objects that can be used safely in command functions.
 	virtual FConcurrentLinearBulkObjectAllocator& GetAllocator() = 0;
 
+	DECLARE_MULTICAST_DELEGATE_OneParam(FOnPreSceneRender, FRDGBuilder&);
+	static ENGINE_API FOnPreSceneRender& OnPreSceneRender();
+
 private:
 	friend class FSceneRenderGroupScope;
 	virtual void BeginGroup(FString&& Name, ESceneRenderGroupFlags Flags) = 0;
diff --git a/Engine/Source/Runtime/Renderer/Private/SceneRenderBuilder.cpp b/Engine/Source/Runtime/Renderer/Private/SceneRenderBuilder.cpp
--- a/Engine/Source/Runtime/Renderer/Private/SceneRenderBuilder.cpp
+++ b/Engine/Source/Runtime/Renderer/Private/SceneRenderBuilder.cpp
@@ -900,6 +900,11 @@ void FSceneRenderProcessor::Execute()
 
 				FRDGBuilder GraphBuilder(RHICmdList, RDG_EVENT_NAME("%s", FunctionInputs.FullPath), ERDGBuilderFlags::Parallel, Scene->GetShaderPlatform());
 
+				if (SceneUpdateInputs)
+				{
+					ISceneRenderBuilder::OnPreSceneRender().Broadcast(GraphBuilder);
+				}
+
 			#if WITH_MGPU
 				if (Renderer.ViewFamily.bForceCopyCrossGPU)
 				{
diff --git a/Engine/Source/Runtime/SlateRHIRenderer/Private/SlateRHIRenderer.cpp b/Engine/Source/Runtime/SlateRHIRenderer/Private/SlateRHIRenderer.cpp
--- a/Engine/Source/Runtime/SlateRHIRenderer/Private/SlateRHIRenderer.cpp
+++ b/Engine/Source/Runtime/SlateRHIRenderer/Private/SlateRHIRenderer.cpp
@@ -56,6 +56,7 @@
 #include "SlatePostProcessor.h"
 #include "Stats/ThreadIdleStats.h"
 #include "VT/VirtualTextureFeedbackResource.h"
+#include "SceneRenderBuilderInterface.h"
 
 #if WITH_EDITORONLY_DATA
 #include "ShaderCompiler.h"
@@ -244,6 +245,11 @@ bool FSlateRHIRenderer::Initialize()
 
 	ElementBatcher = MakeUnique<FSlateElementBatcher>(RenderingPolicy.ToSharedRef());
 
+	DeferredUpdateDelegate = ISceneRenderBuilder::OnPreSceneRender().AddLambda([this] (FRDGBuilder& GraphBuilder)
+	{
+		UpdateDeferredWidgets(GraphBuilder);
+	});
+
 	CurrentSceneIndex = -1;
 	ActiveScenes.Empty();
 	return true;
@@ -277,6 +283,8 @@ void FSlateRHIRenderer::Destroy()
 
 	CurrentSceneIndex = -1;
 	ActiveScenes.Empty();
+
+	ISceneRenderBuilder::OnPreSceneRender().Remove(DeferredUpdateDelegate);
 }
 
 /** Returns a draw buffer that can be used by Slate windows to draw window elements */
@@ -1106,14 +1114,11 @@ void FSlateRHIRenderer::PrepareToTakeHDRScreenshot(const FIntRect& Rect, TArray<
 
 struct FSlateDrawWindowsCommand : public TConcurrentLinearObject<FSlateDrawWindowsCommand>
 {
-	bool IsEmpty() { return Windows.IsEmpty() && DeferredUpdates.IsEmpty(); }
-
 	TArray<FSlateDrawWindowPassInputs, FConcurrentLinearArrayAllocator> Windows;
 	TArray<FSlatePostProcessUpdateRequest, FConcurrentLinearArrayAllocator> PostProcessUpdates;
-	TArray<FRenderThreadUpdateContext, FConcurrentLinearArrayAllocator> DeferredUpdates;
 };
 
-void FSlateRHIRenderer::DrawWindows_RenderThread(FRHICommandListImmediate& RHICmdList, TConstArrayView<FSlateDrawWindowPassInputs> Windows, TConstArrayView<FRenderThreadUpdateContext> DeferredUpdates)
+void FSlateRHIRenderer::DrawWindows_RenderThread(FRHICommandListImmediate& RHICmdList, TConstArrayView<FSlateDrawWindowPassInputs> Windows)
 {
 	struct FWindowPresentCommand
 	{
@@ -1136,10 +1141,7 @@ void FSlateRHIRenderer::DrawWindows_RenderThread(FRHICommandListImmediate& RHICm
 		{
 			FRDGBuilder GraphBuilder(RHICmdList, RDG_EVENT_NAME("Slate"), ERDGBuilderFlags::ParallelSetup | ERDGBuilderFlags::ParallelExecute);
 
-			for (const FRenderThreadUpdateContext& DeferredUpdateContext : DeferredUpdates)
-			{
-				DeferredUpdateContext.Renderer->DrawWindowToTarget_RenderThread(GraphBuilder, DeferredUpdateContext);
-			}
+			UpdateDeferredWidgets(GraphBuilder);
 
 			// D3D12 can't handle more than 8 swap chains at a time, start a new batch if we hit this amount and continue with a new builder.
 			for (int32 NumWindows = 0; NumWindows < 8 && WindowIndex < Windows.Num(); ++NumWindows, ++WindowIndex)
@@ -1152,19 +1154,12 @@ void FSlateRHIRenderer::DrawWindows_RenderThread(FRHICommandListImmediate& RHICm
 			GraphBuilder.Execute();
 		}
 
-		for (const FRenderThreadUpdateContext& DeferredUpdateContext : DeferredUpdates)
-		{
-			DeferredUpdateContext.Renderer->ReleaseDrawBuffer(*DeferredUpdateContext.WindowDrawBuffer);
-		}
-
 		for (const FWindowPresentCommand& Command : WindowPresentCommands)
 		{
 			PresentWindow_RenderThread(RHICmdList, Command.Inputs, Command.Outputs);
 		}
 		WindowPresentCommands.Reset();
 
-		DeferredUpdates = {};
-
 	} while (WindowIndex < Windows.Num());
 }
 
@@ -1359,7 +1354,6 @@ void FSlateRHIRenderer::DrawWindows_Private(FSlateDrawBuffer& WindowDrawBuffer)
 
 	TUniquePtr<FSlateDrawWindowsCommand> DrawWindowsCommand = MakeUnique<FSlateDrawWindowsCommand>();
 	DrawWindowsCommand->Windows.Reserve(WindowsToRender.Num());
-	DrawWindowsCommand->DeferredUpdates = MoveTemp(DeferredUpdateContexts);
 	DrawWindowsCommand->PostProcessUpdates.Reserve(WindowsToRender.Num() * FMath::CountBits((int32)PostProcessAnyUsedBits));
 
 	bool bScreenshotProcessed = false;
@@ -1426,15 +1420,10 @@ void FSlateRHIRenderer::DrawWindows_Private(FSlateDrawBuffer& WindowDrawBuffer)
 		}
 	}
 
-	if (!DrawWindowsCommand->IsEmpty())
+	ENQUEUE_RENDER_COMMAND(SlateDrawWindowsCommand)([this, DrawWindowsCommand = MoveTemp(DrawWindowsCommand)](FRHICommandListImmediate& RHICmdList)
 	{
-		ENQUEUE_RENDER_COMMAND(SlateDrawWindowsCommand)([this, DrawWindowsCommand = MoveTemp(DrawWindowsCommand)](FRHICommandListImmediate& RHICmdList)
-		{
-			DrawWindows_RenderThread(RHICmdList, DrawWindowsCommand->Windows, DrawWindowsCommand->DeferredUpdates);
-		});
-
-		check(DeferredUpdateContexts.IsEmpty());
-	}
+		DrawWindows_RenderThread(RHICmdList, DrawWindowsCommand->Windows);
+	});
 
 	if (bScreenshotProcessed)
 	{
@@ -1795,17 +1784,35 @@ void FSlateRHIRenderer::RequestResize(const TSharedPtr<SWindow>& Window, uint32
 
 void FSlateRHIRenderer::AddWidgetRendererUpdate(const FRenderThreadUpdateContext& Context, bool bDeferredRenderTargetUpdate)
 {
-	if (bDeferredRenderTargetUpdate)
+	ENQUEUE_RENDER_COMMAND(AddWidgetRendererUpdate)([this, Context, bDeferredRenderTargetUpdate](FRHICommandListImmediate& RHICmdList)
 	{
-		DeferredUpdateContexts.Add(Context);
-	}
-	else
-	{
-		ENQUEUE_RENDER_COMMAND(DrawWidgetRendererImmediate)([Context](FRHICommandListImmediate& RHICmdList)
+		if (bDeferredRenderTargetUpdate)
 		{
-			FRDGBuilder GraphBuilder(RHICmdList, RDG_EVENT_NAME("SlateWidgetRender"), ERDGBuilderFlags::ParallelSetup | ERDGBuilderFlags::ParallelExecute);
+			DeferredUpdateContexts.Add(Context);
+		}
+		else
+		{
+			FRDGBuilder GraphBuilder(RHICmdList, RDG_EVENT_NAME("SlateWidgetRender"));
 			Context.Renderer->DrawWindowToTarget_RenderThread(GraphBuilder, Context);
 			GraphBuilder.Execute();
-		});
+		}
+	});
+}
+
+void FSlateRHIRenderer::UpdateDeferredWidgets(FRDGBuilder& GraphBuilder)
+{
+	TArray<FRenderThreadUpdateContext, FConcurrentLinearArrayAllocator> Contexts = MoveTemp(DeferredUpdateContexts);
+
+	for (const FRenderThreadUpdateContext& DeferredUpdateContext : Contexts)
+	{
+		DeferredUpdateContext.Renderer->DrawWindowToTarget_RenderThread(GraphBuilder, DeferredUpdateContext);
 	}
+
+	GraphBuilder.AddPostExecuteCallback([Contexts = MoveTemp(Contexts)]
+	{
+		for (const FRenderThreadUpdateContext& DeferredUpdateContext : Contexts)
+		{
+			DeferredUpdateContext.Renderer->ReleaseDrawBuffer(*DeferredUpdateContext.WindowDrawBuffer);
+		}
+	});
 }
\ No newline at end of file
diff --git a/Engine/Source/Runtime/SlateRHIRenderer/Private/SlateRHIRenderer.h b/Engine/Source/Runtime/SlateRHIRenderer/Private/SlateRHIRenderer.h
--- a/Engine/Source/Runtime/SlateRHIRenderer/Private/SlateRHIRenderer.h
+++ b/Engine/Source/Runtime/SlateRHIRenderer/Private/SlateRHIRenderer.h
@@ -113,7 +113,7 @@ private:
 
 	void PresentWindow_RenderThread(FRHICommandListImmediate& RHICmdList, const FSlateDrawWindowPassInputs& DrawPassInputs, const FSlateDrawWindowPassOutputs& DrawPassOutputs);
 
-	void DrawWindows_RenderThread(FRHICommandListImmediate& RHICmdList, TConstArrayView<FSlateDrawWindowPassInputs> Windows, TConstArrayView<FRenderThreadUpdateContext> DeferredUpdates);
+	void DrawWindows_RenderThread(FRHICommandListImmediate& RHICmdList, TConstArrayView<FSlateDrawWindowPassInputs> Windows);
 
 	/** Loads all known textures from Slate styles */
 	void LoadUsedTextures();
@@ -172,6 +172,9 @@ private:
 	void FlushPendingDeletes();
 
 	TArray<FRenderThreadUpdateContext, FConcurrentLinearArrayAllocator> DeferredUpdateContexts;
+	FDelegateHandle DeferredUpdateDelegate;
+
+	void UpdateDeferredWidgets(FRDGBuilder& GraphBuilder);
 
 	bool bIsStandaloneStereoOnlyDevice = false;
 	bool bUpdateHDRDisplayInformation = false;
diff --git a/Engine/Source/Runtime/UMG/Private/Slate/SRetainerWidget.cpp b/Engine/Source/Runtime/UMG/Private/Slate/SRetainerWidget.cpp
--- a/Engine/Source/Runtime/UMG/Private/Slate/SRetainerWidget.cpp
+++ b/Engine/Source/Runtime/UMG/Private/Slate/SRetainerWidget.cpp
@@ -54,19 +54,11 @@ static bool IsRetainedRenderingEnabled()
 	return GEnableRetainedRendering != 0;
 }
 
-/** Whether or not the platform should have deferred retainer widget render target updating enabled by default */
-#define PLATFORM_REQUIRES_DEFERRED_RETAINER_UPDATE PLATFORM_IOS || PLATFORM_ANDROID;
-
-/**
- * If this is true the retained rendering render thread work will happen during normal slate render thread rendering after the back buffer has been presented
- * in order to avoid extra render target switching in the middle of the frame. The downside is that the UI update will be a frame late
- */
-int32 GDeferRetainedRenderingRenderThread = PLATFORM_REQUIRES_DEFERRED_RETAINER_UPDATE;
+int32 GDeferRetainedRenderingRenderThread = 1;
 FAutoConsoleVariableRef DeferRetainedRenderingRT(
 	TEXT("Slate.DeferRetainedRenderingRenderThread"),
 	GDeferRetainedRenderingRenderThread,
-	TEXT("Whether or not to defer retained rendering to happen at the same time as the rest of slate render thread work"));
-
+	TEXT("Whether or not to defer retained rendering to happen at the same time as the rest of slate render thread work. This improves performance by merging RDG builders and reducing render target swaps."));
 
 class FRetainerWidgetRenderingResources : public FDeferredCleanupInterface, public FGCObject
 {
