Slate - Don't Transition Backbuffer into SRV

Not transitioning into the SRV state may avoid decompresion of backbuffers on platforms that support presenting compressed resources.

diff --git a/Engine/Source/Runtime/SlateRHIRenderer/Private/SlateRHIRenderer.cpp b/Engine/Source/Runtime/SlateRHIRenderer/Private/SlateRHIRenderer.cpp
--- a/Engine/Source/Runtime/SlateRHIRenderer/Private/SlateRHIRenderer.cpp
+++ b/Engine/Source/Runtime/SlateRHIRenderer/Private/SlateRHIRenderer.cpp
@@ -906,6 +906,10 @@ FSlateDrawWindowPassOutputs FSlateRHIRenderer::DrawWindow_RenderThread(FRDGBuild
 			GraphBuilder.SetTextureAccessFinal(SwapChainTexture, ERHIAccess::Present);
 			GEngine->StereoRenderingDevice->RenderTexture_RenderThread(GraphBuilder, SwapChainTexture, OutputTexture, WindowElementList.GetWindowSize());
 		}
+		else
+		{
+			GraphBuilder.SetTextureAccessFinal(OutputTexture, ERHIAccess::Present);
+		}
 
 		if (bRequiresVirtualTextureFeedback)
 		{
@@ -981,16 +985,6 @@ void FSlateRHIRenderer::PresentWindow_RenderThread(FRHICommandListImmediate& RHI
 	}
 
 	RHICmdList.EndDrawingViewport(DrawPassOutputs.ViewportRHI, true, DrawPassInputs.bLockToVsync);
-	
-	if (GRHIGlobals.NeedsExtraTransitions)
-	{
-		RHICmdList.Transition(FRHITransitionInfo(DrawPassOutputs.OutputTextureRHI, ERHIAccess::Present, ERHIAccess::SRVMask));
-
-		if (OptionalSDRBuffer)
-		{
-			RHICmdList.Transition(FRHITransitionInfo(OptionalSDRBuffer, ERHIAccess::Present, ERHIAccess::SRVMask));
-		}
-	}
 
 	RHICmdList.EnqueueLambda([CurrentFrameCounter = GFrameCounterRenderThread](FRHICommandListImmediate& InRHICmdList)
 	{
