Fix random crash in SlateRHIRenderer due to both main thread and render thread possibly writing to the same arrays simultaneously.
Some renaming for clarity.

#tests This is a rare crash with no repro, so I can't be 100% sure it's fixed. Ran engine test client, game client/server, and editor with the new code, ensuring that invalidation was enabled, with all error detections enabled for the invalidation system.
[RN] fix ui/slate

diff --git a/Engine/Source/Runtime/SlateCore/Public/FastUpdate/WidgetProxy.h b/Engine/Source/Runtime/SlateCore/Public/FastUpdate/WidgetProxy.h
--- a/Engine/Source/Runtime/SlateCore/Public/FastUpdate/WidgetProxy.h
+++ b/Engine/Source/Runtime/SlateCore/Public/FastUpdate/WidgetProxy.h
@@ -129,7 +129,7 @@ public:
 	};
 
 	/**
-	 * Similar to SWidget::Paint but use the saved PersistenState. Only paint/tick if needed.
+	 * Similar to SWidget::Paint but use the saved PersistentState. Only paint/tick if needed.
 	 * @return the new and previous LayerId is the widget was painted.
 	 */
 	FUpdateResult Update(const FPaintArgs& PaintArgs, FSlateWindowElementList& OutDrawElements);
diff --git a/Engine/Source/Runtime/SlateRHIRenderer/Private/SlateRHIRenderer.cpp b/Engine/Source/Runtime/SlateRHIRenderer/Private/SlateRHIRenderer.cpp
--- a/Engine/Source/Runtime/SlateRHIRenderer/Private/SlateRHIRenderer.cpp
+++ b/Engine/Source/Runtime/SlateRHIRenderer/Private/SlateRHIRenderer.cpp
@@ -1688,30 +1688,33 @@ EPixelFormat FSlateRHIRenderer::GetSlateRecommendedColorFormat()
 	return bIsStandaloneStereoOnlyDevice ? PF_R8G8B8A8 : PF_B8G8R8A8;
 }
 
-void FSlateRHIRenderer::DestroyCachedFastPathRenderingData(FSlateCachedFastPathRenderingData* CachedRenderingData)
+void FSlateRHIRenderer::DestroyCachedFastPathRenderingData(FSlateCachedFastPathRenderingData* CachedFastPathRenderingData)
 {
-	check(CachedRenderingData);
-	PendingDeletes.CachedRenderingData.Emplace(CachedRenderingData);
+	check(CachedFastPathRenderingData);
+	FScopeLock ScopeLock(GetResourceCriticalSection());
+	PendingDeletes.CachedFastPathRenderingData.Emplace(CachedFastPathRenderingData);
 }
 
 void FSlateRHIRenderer::DestroyCachedFastPathElementData(FSlateCachedElementData* CachedElementData)
 {
 	check(CachedElementData);
+	FScopeLock ScopeLock(GetResourceCriticalSection());
 	PendingDeletes.CachedElementData.Emplace(CachedElementData);
 }
 
 void FSlateRHIRenderer::FlushPendingDeletes()
 {
+	FScopeLock ScopeLock(GetResourceCriticalSection());
 	if (!PendingDeletes.IsEmpty())
 	{
-		ENQUEUE_RENDER_COMMAND(SlateDeferredDelete)([PendingDeletes = MoveTemp(PendingDeletes)](FRHICommandListBase&)
+		ENQUEUE_RENDER_COMMAND(SlateDeferredDelete)([InPendingDeletes = MoveTemp(PendingDeletes)](FRHICommandListBase&)
 		{
-			for (FSlateCachedFastPathRenderingData* Data : PendingDeletes.CachedRenderingData)
+			for (FSlateCachedFastPathRenderingData* Data : InPendingDeletes.CachedFastPathRenderingData)
 			{
 				delete Data;
 			}
 
-			for (FSlateCachedElementData* Data : PendingDeletes.CachedElementData)
+			for (FSlateCachedElementData* Data : InPendingDeletes.CachedElementData)
 			{
 				delete Data;
 			}
diff --git a/Engine/Source/Runtime/SlateRHIRenderer/Private/SlateRHIRenderer.h b/Engine/Source/Runtime/SlateRHIRenderer/Private/SlateRHIRenderer.h
--- a/Engine/Source/Runtime/SlateRHIRenderer/Private/SlateRHIRenderer.h
+++ b/Engine/Source/Runtime/SlateRHIRenderer/Private/SlateRHIRenderer.h
@@ -162,10 +162,10 @@ private:
 
 	struct
 	{
-		TArray<FSlateCachedFastPathRenderingData*, FConcurrentLinearArrayAllocator> CachedRenderingData;
+		TArray<FSlateCachedFastPathRenderingData*, FConcurrentLinearArrayAllocator> CachedFastPathRenderingData;
 		TArray<FSlateCachedElementData*, FConcurrentLinearArrayAllocator> CachedElementData;
 
-		bool IsEmpty() const { return CachedRenderingData.IsEmpty() && CachedElementData.IsEmpty(); }
+		bool IsEmpty() const { return CachedFastPathRenderingData.IsEmpty() && CachedElementData.IsEmpty(); }
 
 	} PendingDeletes;
 
