Add a CSV custom stat for Slate postprocess blur counts

diff --git a/Engine/Source/Runtime/SlateRHIRenderer/Private/SlatePostProcessor.cpp b/Engine/Source/Runtime/SlateRHIRenderer/Private/SlatePostProcessor.cpp
--- a/Engine/Source/Runtime/SlateRHIRenderer/Private/SlatePostProcessor.cpp
+++ b/Engine/Source/Runtime/SlateRHIRenderer/Private/SlatePostProcessor.cpp
@@ -7,6 +7,8 @@
 
 //////////////////////////////////////////////////////////////////////////
 
+CSV_DECLARE_CATEGORY_MODULE_EXTERN(SLATECORE_API, Slate);
+
 float GetSlateHDRUILevel()
 {
 	static auto CVar = IConsoleManager::Get().FindConsoleVariable(TEXT("r.HDR.UI.Level"));
@@ -313,6 +315,7 @@ IMPLEMENT_GLOBAL_SHADER(FSlatePostProcessBlurPS, "/Engine/Private/SlatePostProce
 void AddSlatePostProcessBlurPass(FRDGBuilder& GraphBuilder, const FSlatePostProcessBlurPassInputs& Inputs)
 {
 	RDG_EVENT_SCOPE(GraphBuilder, "GaussianBlur");
+	CSV_CUSTOM_STAT(Slate, PostProcessBlurPassCount, 1, ECsvCustomStatOp::Accumulate);
 
 	const auto GetWeight = [](float Dist, float Strength)
 	{
