Slate: Fix Slate.ShowBatching regression caused by 35454935. Batch colors assignment in CreateInputs variable was moved outside of RenderBatchIndex iteration.

#rnx

diff --git a/Engine/Source/Runtime/SlateRHIRenderer/Private/SlateRHIRenderingPolicy.cpp b/Engine/Source/Runtime/SlateRHIRenderer/Private/SlateRHIRenderingPolicy.cpp
--- a/Engine/Source/Runtime/SlateRHIRenderer/Private/SlateRHIRenderingPolicy.cpp
+++ b/Engine/Source/Runtime/SlateRHIRenderer/Private/SlateRHIRenderingPolicy.cpp
@@ -1435,7 +1435,7 @@ void AddSlateDrawElementsPass(
 	FRandomStream BatchColors(1337);
 #endif
 
-	const FSlateRenderBatchCreateInputs RenderBatchCreateInputs
+	FSlateRenderBatchCreateInputs RenderBatchCreateInputs
 	{
 		  .ShaderMap             = ShaderMap
 		, .SceneViewAllocator    = SceneViewAllocator
@@ -1665,6 +1665,11 @@ void AddSlateDrawElementsPass(
 					RenderBatchTailOp = RenderBatchOp;
 				}
 				NumRenderBatchOps++;
+
+#if WITH_SLATE_VISUALIZERS
+				// Batch Color was used, move to new random color for next batch.
+				RenderBatchCreateInputs.BatchColor = FLinearColor(BatchColors.GetUnitVector());
+#endif
 			}
 
 			break;
