Add native MakeStruct function for SlateFontInfo to validate the type of font material

diff --git a/Engine/Source/Runtime/SlateCore/Public/Fonts/SlateFontInfo.h b/Engine/Source/Runtime/SlateCore/Public/Fonts/SlateFontInfo.h
--- a/Engine/Source/Runtime/SlateCore/Public/Fonts/SlateFontInfo.h
+++ b/Engine/Source/Runtime/SlateCore/Public/Fonts/SlateFontInfo.h
@@ -142,7 +142,7 @@ struct TStructOpsTypeTraits<FFontOutlineSettings>
 /**
  * A representation of a font in Slate.
  */
-USTRUCT(BlueprintType)
+USTRUCT(BlueprintType, meta = (HasNativeMake = "/Script/SlateRHIRenderer.SlateFontBlueprintLibrary.MakeSlateFontInfo"))
 struct FSlateFontInfo
 {
 	GENERATED_USTRUCT_BODY()
diff --git a/Engine/Source/Runtime/SlateRHIRenderer/Private/SlateFontBlueprintLibrary.cpp b/Engine/Source/Runtime/SlateRHIRenderer/Private/SlateFontBlueprintLibrary.cpp
new file mode 100644
--- /dev/null
+++ b/Engine/Source/Runtime/SlateRHIRenderer/Private/SlateFontBlueprintLibrary.cpp
@@ -0,0 +1,42 @@
+// Copyright Epic Games, Inc. All Rights Reserved.
+
+#include "SlateFontBlueprintLibrary.h"
+#include "Blueprint/BlueprintExceptionInfo.h"
+#include "Materials/MaterialInterface.h"
+
+DEFINE_LOG_CATEGORY(LogSlateFontBlueprint);
+
+
+FSlateFontInfo USlateFontBlueprintLibrary::MakeSlateFontInfo(const UObject* FontObject, UObject* FontMaterial, FFontOutlineSettings OutlineSettings, FName TypefaceFontName, float Size, int32 LetterSpacing, float SkewAmount, bool bForceMonospaced, bool bMaterialIsStencil, float MonospacedWidth)
+{
+	FSlateFontInfo SlateFontInfo = FSlateFontInfo();
+	SlateFontInfo.FontObject = FontObject;
+	SlateFontInfo.FontMaterial = FontMaterial;
+	SlateFontInfo.OutlineSettings = OutlineSettings;
+	SlateFontInfo.TypefaceFontName = TypefaceFontName;
+	SlateFontInfo.Size = Size;
+	SlateFontInfo.LetterSpacing = LetterSpacing;
+	SlateFontInfo.SkewAmount = SkewAmount;
+	SlateFontInfo.bForceMonospaced = bForceMonospaced;
+	SlateFontInfo.bMaterialIsStencil = bMaterialIsStencil;
+	SlateFontInfo.MonospacedWidth = MonospacedWidth;
+
+	if (FontMaterial && !FontMaterial->IsA<UMaterialInterface>())
+	{
+#if !UE_BUILD_SHIPPING && !UE_BUILD_TEST
+		if (FFrame::GetThreadLocalTopStackFrame() && FFrame::GetThreadLocalTopStackFrame()->Object)
+		{
+			FText ErrorMessage = FText::FromString(TEXT("Font Material is not of type material. Either remove the reference or replace it with an object of type material."));
+			const FBlueprintExceptionInfo ExceptionInfo(EBlueprintExceptionType::FatalError, ErrorMessage);
+			FBlueprintCoreDelegates::ThrowScriptException(FFrame::GetThreadLocalTopStackFrame()->Object, *FFrame::GetThreadLocalTopStackFrame(), ExceptionInfo);
+		}
+		else
+#endif
+		{
+			UE_LOG(LogSlateFontBlueprint, Warning, TEXT("Font Material is not of type material. It will be replaced with an empty material."));
+		}
+		SlateFontInfo.FontObject = nullptr;
+	}
+
+	return SlateFontInfo;
+}
\ No newline at end of file
diff --git a/Engine/Source/Runtime/SlateRHIRenderer/Private/SlateFontBlueprintLibrary.h b/Engine/Source/Runtime/SlateRHIRenderer/Private/SlateFontBlueprintLibrary.h
new file mode 100644
--- /dev/null
+++ b/Engine/Source/Runtime/SlateRHIRenderer/Private/SlateFontBlueprintLibrary.h
@@ -0,0 +1,18 @@
+// Copyright Epic Games, Inc. All Rights Reserved.
+
+#pragma once
+
+#include "Fonts/SlateFontInfo.h"
+#include "UObject/Class.h"
+#include "SlateFontBlueprintLibrary.generated.h"
+
+SLATERHIRENDERER_API DECLARE_LOG_CATEGORY_EXTERN(LogSlateFontBlueprint, Log, All);
+
+UCLASS()
+class USlateFontBlueprintLibrary : public UObject
+{
+	GENERATED_BODY()
+
+	UFUNCTION(BlueprintPure, Category = "Slate Font", meta = (NativeMakeFunc))
+	static SLATERHIRENDERER_API FSlateFontInfo MakeSlateFontInfo(const UObject* FontObject, UObject* FontMaterial, FFontOutlineSettings OutlineSettings, FName TypefaceFontName, float Size = 24.0f, int32 LetterSpacing = 0, float SkewAmount = 0.0f, bool bForceMonospaced = false, bool bMaterialIsStencil = false, float MonospacedWidth = 1.0f);
+};
\ No newline at end of file
diff --git a/Engine/Source/Runtime/SlateRHIRenderer/Private/SlateRHIResourceManager.cpp b/Engine/Source/Runtime/SlateRHIRenderer/Private/SlateRHIResourceManager.cpp
--- a/Engine/Source/Runtime/SlateRHIRenderer/Private/SlateRHIResourceManager.cpp
+++ b/Engine/Source/Runtime/SlateRHIRenderer/Private/SlateRHIResourceManager.cpp
@@ -676,6 +676,22 @@ FSlateShaderResource* FSlateRHIResourceManager::GetFontShaderResource( int32 InT
 	{
 		return FontTextureAtlas;
 	}
+	else if (!FontMaterial->IsA<UMaterialInterface>())
+	{
+#if SLATE_CHECK_UOBJECT_RENDER_RESOURCES
+		FString ErrorMsg = FString::Format(TEXT("Cannot render object '{0}' as a material because it's of type {1} instead of type MaterialInterface. Make sure all materials passed for rendering in slate are Material Interfaces. This is caused by bad user data (a UObject that is not a UMaterialInterface) passed for rendering."), { FontMaterial->GetFullName(), FontMaterial->GetClass()->GetName() });
+
+		if (GSlateCheckUObjectRenderResourcesShouldLogFatal)
+		{
+			UE_LOG(LogSlate, Fatal, TEXT("%s"), *ErrorMsg);
+		}
+		else
+		{
+			ensureMsgf(false, TEXT("%s"), *ErrorMsg);
+		}
+#endif
+		return FontTextureAtlas;
+	}
 	else
 	{
 		return GetMaterialResource( FontMaterial, nullptr, FontTextureAtlas, InTextureAtlasIndex );
