SlateFontBlueprint: when the material is invalid, clear the material instead of the font object

diff --git a/Engine/Source/Runtime/SlateRHIRenderer/Private/SlateFontBlueprintLibrary.cpp b/Engine/Source/Runtime/SlateRHIRenderer/Private/SlateFontBlueprintLibrary.cpp
--- a/Engine/Source/Runtime/SlateRHIRenderer/Private/SlateFontBlueprintLibrary.cpp
+++ b/Engine/Source/Runtime/SlateRHIRenderer/Private/SlateFontBlueprintLibrary.cpp
@@ -35,7 +35,7 @@ FSlateFontInfo USlateFontBlueprintLibrary::MakeSlateFontInfo(const UObject* Font
 		{
 			UE_LOG(LogSlateFontBlueprint, Warning, TEXT("Font Material is not of type material. It will be replaced with an empty material."));
 		}
-		SlateFontInfo.FontObject = nullptr;
+		SlateFontInfo.FontMaterial = nullptr;
 	}
 
 	return SlateFontInfo;
