Add r.HDR.UI.CompositeEOTF to enable choice of "EOTF" for UI composition in HDR and change the default behavior from sRGB to BT.1886

Viewing environment standards for Rec. 709 expect BT.1886 for EOTF and are similar to the viewing environment specified in Rec. 2100 for HDR.

diff --git a/Engine/Shaders/Private/CompositeUIPixelShader.usf b/Engine/Shaders/Private/CompositeUIPixelShader.usf
--- a/Engine/Shaders/Private/CompositeUIPixelShader.usf
+++ b/Engine/Shaders/Private/CompositeUIPixelShader.usf
@@ -32,21 +32,12 @@ SamplerState SceneSampler;
 
 float UILevel;
 float UILuminance;
+int UICompositeEOTF;
 
-// Invert the encoding applied by the ApplyGammaCorrection function.
+// Apply "EOTF" on SDR UI for HDR composition.
 float3 LinearizeColor(float3 EncodedColor)
 {
-	#if MAC
-		// Note, MacOSX native output is raw gamma 2.2 not sRGB!
-		return pow(EncodedColor, 2.2);
-	#else
-		#if USE_709
-			// Didn't profile yet if the branching version would be faster (different linear segment).
-			return Rec709ToLinear(EncodedColor);
-		#else
-			return sRGBToLinear(EncodedColor);
-		#endif
-	#endif
+	return (UICompositeEOTF == 0) ? pow(EncodedColor, 2.4) : sRGBToLinear(EncodedColor);
 }
 
 static const float ScRGBScaleFactor = 80.0f; // used to convert between ScRGB and nits
diff --git a/Engine/Source/Runtime/SlateRHIRenderer/Private/SlatePostProcessor.cpp b/Engine/Source/Runtime/SlateRHIRenderer/Private/SlatePostProcessor.cpp
--- a/Engine/Source/Runtime/SlateRHIRenderer/Private/SlatePostProcessor.cpp
+++ b/Engine/Source/Runtime/SlateRHIRenderer/Private/SlatePostProcessor.cpp
@@ -21,6 +21,12 @@ float GetSlateHDRUILuminance()
 	return CVar ? CVar->GetFloat() : 300.0f;
 }
 
+int GetSlateHDRUICompositeEOTF()
+{
+	static auto CVar = IConsoleManager::Get().FindConsoleVariable(TEXT("r.HDR.UI.CompositeEOTF"));
+	return CVar ? CVar->GetInt() : 0;
+}
+
 ETextureCreateFlags GetSlateTransientRenderTargetFlags()
 {
 	return ETextureCreateFlags::RenderTargetable | ETextureCreateFlags::ShaderResource | ETextureCreateFlags::FastVRAM
@@ -55,6 +61,7 @@ public:
 		SHADER_PARAMETER_SAMPLER(SamplerState, UISampler)
 		SHADER_PARAMETER(float, UILevel)
 		SHADER_PARAMETER(float, UILuminance)
+		SHADER_PARAMETER(int32, UICompositeEOTF)
 		SHADER_PARAMETER(FVector2f, UITextureSize)
 		RENDER_TARGET_BINDING_SLOTS()
 	END_SHADER_PARAMETER_STRUCT()
@@ -113,6 +120,7 @@ FScreenPassTexture AddSlateCompositeHDRForBlurPass(FRDGBuilder& GraphBuilder, co
 	PassParameters->UITextureSize = InputViewport.Extent;
 	PassParameters->UILevel = GetSlateHDRUILevel();
 	PassParameters->UILuminance = GetSlateHDRUILuminance();
+	PassParameters->UICompositeEOTF = GetSlateHDRUICompositeEOTF();
 
 	TShaderMapRef<FCompositeHDRForBlurPS> PixelShader(ShaderMap, PermutationVector);
 	AddDrawScreenPass(GraphBuilder, RDG_EVENT_NAME("CompositeHDR"), FeatureLevel, OutputViewport, InputViewport, PixelShader, PassParameters);
diff --git a/Engine/Source/Runtime/SlateRHIRenderer/Private/SlateRHIRenderer.cpp b/Engine/Source/Runtime/SlateRHIRenderer/Private/SlateRHIRenderer.cpp
--- a/Engine/Source/Runtime/SlateRHIRenderer/Private/SlateRHIRenderer.cpp
+++ b/Engine/Source/Runtime/SlateRHIRenderer/Private/SlateRHIRenderer.cpp
@@ -84,6 +84,14 @@ static TAutoConsoleVariable<float> CVarHDRUILuminance(
 	TEXT("Base Luminance in nits for UI elements when compositing into HDR framebuffer. Gets multiplied by r.HDR.UI.Level"),
 	ECVF_RenderThreadSafe);
 
+static TAutoConsoleVariable<int> CVarHDRUICompositeEOTF(
+	TEXT("r.HDR.UI.CompositeEOTF"),
+	0,
+	TEXT("EOTF used when compositing the UI layer:\n")
+	TEXT("0: BT.1886 to match similar viewing conditions specified by Rec. 2100\n")
+	TEXT("1: sRGB\n"),
+	ECVF_RenderThreadSafe);
+
 static TAutoConsoleVariable<int32> CVarUICompositeMode(
 	TEXT("r.HDR.UI.CompositeMode"),
 	1,
@@ -582,6 +590,7 @@ BEGIN_SHADER_PARAMETER_STRUCT(FCompositeShaderCommonParameters, )
 	SHADER_PARAMETER_RDG_TEXTURE(Texture2D<uint>, UIWriteMaskTexture)
 	SHADER_PARAMETER(float, UILevel)
 	SHADER_PARAMETER(float, UILuminance)
+	SHADER_PARAMETER(int32, UICompositeEOTF)
 	SHADER_PARAMETER(float, ColorVisionDeficiencyType)
 	SHADER_PARAMETER(float, ColorVisionDeficiencySeverity)
 	SHADER_PARAMETER(float, bCorrectDeficiency)
@@ -850,6 +859,7 @@ FSlateDrawWindowPassOutputs FSlateRHIRenderer::DrawWindow_RenderThread(FRDGBuild
 			CommonParameters.UISampler = TStaticSamplerState<SF_Point>::GetRHI();
 			CommonParameters.UILevel = CVarUILevel.GetValueOnRenderThread();
 			CommonParameters.UILuminance = CVarHDRUILuminance.GetValueOnRenderThread();
+			CommonParameters.UICompositeEOTF = CVarHDRUICompositeEOTF.GetValueOnRenderThread();
 			CommonParameters.ColorVisionDeficiencySeverity = (float)GSlateColorDeficiencySeverity;
 			CommonParameters.ColorVisionDeficiencyType = (float)GSlateColorDeficiencyType;
 			CommonParameters.bSimulateCorrectionWithDeficiency = GSlateShowColorDeficiencyCorrectionWithDeficiency ? 1.0f : 0.0f;
