Improve help text for r.HDR.UI.CompositeEOTF

diff --git a/Engine/Source/Runtime/SlateRHIRenderer/Private/SlateRHIRenderer.cpp b/Engine/Source/Runtime/SlateRHIRenderer/Private/SlateRHIRenderer.cpp
--- a/Engine/Source/Runtime/SlateRHIRenderer/Private/SlateRHIRenderer.cpp
+++ b/Engine/Source/Runtime/SlateRHIRenderer/Private/SlateRHIRenderer.cpp
@@ -88,8 +88,8 @@ static TAutoConsoleVariable<int> CVarHDRUICompositeEOTF(
 	TEXT("r.HDR.UI.CompositeEOTF"),
 	0,
 	TEXT("EOTF used when compositing the UI layer:\n")
-	TEXT("0: BT.1886 to match similar viewing conditions specified by Rec. 2100\n")
-	TEXT("1: sRGB\n"),
+	TEXT("0: BT.1886 to approximate the viewing environment specified by Rec. 2100 (default)\n")
+	TEXT("1: sRGB"),
 	ECVF_RenderThreadSafe);
 
 static TAutoConsoleVariable<int32> CVarUICompositeMode(
