Restore PSO precaching for UI materials if the shader platform doesn't support manual vertex fetch and set up FRegisterPSOCollectorCreateFunction for EShadingPath::Mobile

diff --git a/Engine/Source/Runtime/Engine/Private/Materials/Material.cpp b/Engine/Source/Runtime/Engine/Private/Materials/Material.cpp
--- a/Engine/Source/Runtime/Engine/Private/Materials/Material.cpp
+++ b/Engine/Source/Runtime/Engine/Private/Materials/Material.cpp
@@ -4794,7 +4794,7 @@ void UMaterial::PostLoad()
 	else if (IsDeferredDecal() || IsUIMaterial() || (MaterialDomain == MD_LightFunction) || IsPostProcessMaterial())
 	{
 		// TODO: need to pass a correct vertex declaration for non-MVF platforms
-		if (RHISupportsManualVertexFetch(GMaxRHIShaderPlatform)) 
+		if (IsUIMaterial() || RHISupportsManualVertexFetch(GMaxRHIShaderPlatform))
 		{
 			FPSOPrecacheParams PSOPrecacheParams;
 			UMaterialInterface::PrecachePSOs(&FLocalVertexFactory::StaticType, PSOPrecacheParams);
diff --git a/Engine/Source/Runtime/Engine/Private/Materials/MaterialInstance.cpp b/Engine/Source/Runtime/Engine/Private/Materials/MaterialInstance.cpp
--- a/Engine/Source/Runtime/Engine/Private/Materials/MaterialInstance.cpp
+++ b/Engine/Source/Runtime/Engine/Private/Materials/MaterialInstance.cpp
@@ -3651,7 +3651,7 @@ void UMaterialInstance::PostLoad()
 	else if (IsDeferredDecal() || IsUIMaterial() || IsPostProcessMaterial())
 	{
 		// TODO: need to pass a correct vertex declaration for non-MVF platforms
-		if (RHISupportsManualVertexFetch(GMaxRHIShaderPlatform))
+		if (IsUIMaterial() || RHISupportsManualVertexFetch(GMaxRHIShaderPlatform))
 		{
 			FPSOPrecacheParams PSOPrecacheParams;
 			UMaterialInterface::PrecachePSOs(&FLocalVertexFactory::StaticType, PSOPrecacheParams);
diff --git a/Engine/Source/Runtime/SlateRHIRenderer/Private/SlateRHIRenderingPolicy.cpp b/Engine/Source/Runtime/SlateRHIRenderer/Private/SlateRHIRenderingPolicy.cpp
--- a/Engine/Source/Runtime/SlateRHIRenderer/Private/SlateRHIRenderingPolicy.cpp
+++ b/Engine/Source/Runtime/SlateRHIRenderer/Private/SlateRHIRenderingPolicy.cpp
@@ -14,6 +14,7 @@
 #include "SceneView.h"
 #include "SceneUtils.h"
 #include "Engine/Engine.h"
+#include "Engine/RendererSettings.h"
 #include "SlateShaders.h"
 #include "Rendering/SlateRenderer.h"
 #include "SlateRHIRenderer.h"
@@ -1716,6 +1717,12 @@ static TArray<ESlateShader> SlateShaderTypesToPrecache =
 	ESlateShader::RoundedBox
 };
 
+static EPixelFormat GetDefaultBackBufferPixelFormat()
+{
+	static const auto CVarDefaultBackBufferPixelFormat = IConsoleManager::Get().FindTConsoleVariableDataInt(TEXT("r.DefaultBackBufferPixelFormat"));
+	return EDefaultBackBufferPixelFormat::Convert2PixelFormat(EDefaultBackBufferPixelFormat::FromInt(CVarDefaultBackBufferPixelFormat->GetValueOnGameThread()));
+}
+
 void AddSlatePSOInitializer(
 	FRHIBlendState* BlendState, 
 	bool bInstanced, 
@@ -1740,7 +1747,7 @@ void AddSlatePSOInitializer(
 	GraphicsPSOInit.PrimitiveType = GetRHIPrimitiveType(DrawPrimitiveType);
 
 	FGraphicsPipelineRenderTargetsInfo RenderTargetsInfo;
-	AddRenderTargetInfo(PF_A2B10G10R10, ETextureCreateFlags::RenderTargetable, RenderTargetsInfo);
+	AddRenderTargetInfo(GetDefaultBackBufferPixelFormat(), ETextureCreateFlags::RenderTargetable, RenderTargetsInfo);
 	RenderTargetsInfo.NumSamples = 1;
 
 	GraphicsPSOInit.StatePrecachePSOHash = RHIComputeStatePrecachePSOHash(GraphicsPSOInit);		
@@ -1815,8 +1822,7 @@ class FSlateMaterialPSOCollector : public IPSOCollector
 {
 public:
 	FSlateMaterialPSOCollector(ERHIFeatureLevel::Type InFeatureLevel) : 
-		IPSOCollector(FPSOCollectorCreateManager::GetIndex(GetFeatureLevelShadingPath(InFeatureLevel), SlateMaterialPSOCollectorName)),
-		FeatureLevel(InFeatureLevel)
+		IPSOCollector(FPSOCollectorCreateManager::GetIndex(GetFeatureLevelShadingPath(InFeatureLevel), SlateMaterialPSOCollectorName))
 	{
 	}
 
@@ -1827,10 +1833,6 @@ public:
 		const FPSOPrecacheParams& PreCacheParams,
 		TArray<FPSOPrecacheData>& PSOInitializers
 	) override final;
-
-private:
-
-	ERHIFeatureLevel::Type FeatureLevel;
 };
 
 void FSlateMaterialPSOCollector::CollectPSOInitializers(
@@ -1897,4 +1899,4 @@ IPSOCollector* CreateSlateMaterialPSOCollector(ERHIFeatureLevel::Type FeatureLev
 	return new FSlateMaterialPSOCollector(FeatureLevel);
 }
 FRegisterPSOCollectorCreateFunction RegisterSlateMaterialPSOCollector(&CreateSlateMaterialPSOCollector, EShadingPath::Deferred, SlateMaterialPSOCollectorName);
-
+FRegisterPSOCollectorCreateFunction RegisterMobileSlateMaterialPSOCollector(&CreateSlateMaterialPSOCollector, EShadingPath::Mobile, SlateMaterialPSOCollectorName);
