Update FSlateRHIRenderer::GetViewportPixelFormat to stop overriding the pixel format for ES31 on desktop platforms

diff --git a/Engine/Source/Runtime/SlateRHIRenderer/Private/SlateRHIRenderer.cpp b/Engine/Source/Runtime/SlateRHIRenderer/Private/SlateRHIRenderer.cpp
--- a/Engine/Source/Runtime/SlateRHIRenderer/Private/SlateRHIRenderer.cpp
+++ b/Engine/Source/Runtime/SlateRHIRenderer/Private/SlateRHIRenderer.cpp
@@ -477,7 +477,10 @@ EPixelFormat FSlateRHIRenderer::GetViewportPixelFormat(const SWindow& Window, bo
 	}
 
 	// Use a known default format in VR / Mobile / Transparent Window SDR configurations.
-	if (bIsStandaloneStereoOnlyDevice || GMaxRHIFeatureLevel == ERHIFeatureLevel::ES3_1
+	if (bIsStandaloneStereoOnlyDevice
+#if !PLATFORM_DESKTOP
+		|| GMaxRHIFeatureLevel == ERHIFeatureLevel::ES3_1
+#endif
 #if ALPHA_BLENDED_WINDOWS
 		|| Window.GetTransparencySupport() == EWindowTransparency::PerPixel
 #endif
