remove es_3_1 specific code that was overriding the swapchain format: the original change (17299603) is not needed anymore. If the swapchain does not match the backbuffer format, we will first resolve SceneColorMS to a non MS rendertarget and then draw to the swapchain via a neutral tonemap.
Setting r.DefaultBackBufferPixelFormat to the scene format as SceneColorMS will prevent that additional copy to happen (for example when mobile hdr is 0)

diff --git a/Engine/Source/Runtime/SlateRHIRenderer/Private/SlateRHIRenderer.cpp b/Engine/Source/Runtime/SlateRHIRenderer/Private/SlateRHIRenderer.cpp
--- a/Engine/Source/Runtime/SlateRHIRenderer/Private/SlateRHIRenderer.cpp
+++ b/Engine/Source/Runtime/SlateRHIRenderer/Private/SlateRHIRenderer.cpp
@@ -478,9 +478,6 @@ EPixelFormat FSlateRHIRenderer::GetViewportPixelFormat(const SWindow& Window, bo
 
 	// Use a known default format in VR / Mobile / Transparent Window SDR configurations.
 	if (bIsStandaloneStereoOnlyDevice
-#if !PLATFORM_DESKTOP
-		|| GMaxRHIFeatureLevel == ERHIFeatureLevel::ES3_1
-#endif
 #if ALPHA_BLENDED_WINDOWS
 		|| Window.GetTransparencySupport() == EWindowTransparency::PerPixel
 #endif
