Added more specific logging to figure out why slate brush is not resolving to a texture type.

diff --git a/Engine/Source/Runtime/SlateRHIRenderer/Private/SlateRHIResourceManager.cpp b/Engine/Source/Runtime/SlateRHIRenderer/Private/SlateRHIResourceManager.cpp
--- a/Engine/Source/Runtime/SlateRHIRenderer/Private/SlateRHIResourceManager.cpp
+++ b/Engine/Source/Runtime/SlateRHIRenderer/Private/SlateRHIResourceManager.cpp
@@ -863,7 +863,18 @@ FSlateShaderResourceProxy* FSlateRHIResourceManager::FindOrCreateDynamicTextureR
 				if ( !FailedTextures.Contains(ResourceObject) )
 				{
 					FailedTextures.Add(ResourceObject);
-					ensureMsgf(false, TEXT("Slate RHI Error - Invalid Texture2D '%s'."), *ResourceName.ToString());
+
+					if (IsValid(ResourceObject))
+					{
+						ensureMsgf(false, TEXT("Slate RHI Error - Invalid Texture2D, ResourceName: '%s'. UObject not valid."), *ResourceName.ToString());
+					}
+					else
+					{
+						ensureMsgf(false, TEXT("Slate RHI Error - Invalid Texture2D, ResourceName: '%s', Path: '%s', Class: '%s'"),
+							*ResourceName.ToString(),
+							*ResourceObject->GetFullName(),
+							*ResourceObject->GetClass()->GetName());
+					}
 				}
 				ResourceObject = GetBadResourceTexture();
 #else
