Fixed incorrect IsValid check in slate logging change.

diff --git a/Engine/Source/Runtime/SlateRHIRenderer/Private/SlateRHIResourceManager.cpp b/Engine/Source/Runtime/SlateRHIRenderer/Private/SlateRHIResourceManager.cpp
--- a/Engine/Source/Runtime/SlateRHIRenderer/Private/SlateRHIResourceManager.cpp
+++ b/Engine/Source/Runtime/SlateRHIRenderer/Private/SlateRHIResourceManager.cpp
@@ -865,16 +865,16 @@ FSlateShaderResourceProxy* FSlateRHIResourceManager::FindOrCreateDynamicTextureR
 					FailedTextures.Add(ResourceObject);
 
 					if (IsValid(ResourceObject))
-					{
-						ensureMsgf(false, TEXT("Slate RHI Error - Invalid Texture2D, ResourceName: '%s'. UObject not valid."), *ResourceName.ToString());
-					}
-					else
 					{
 						ensureMsgf(false, TEXT("Slate RHI Error - Invalid Texture2D, ResourceName: '%s', Path: '%s', Class: '%s'"),
 							*ResourceName.ToString(),
 							*ResourceObject->GetFullName(),
 							*ResourceObject->GetClass()->GetName());
 					}
+					else
+					{
+						ensureMsgf(false, TEXT("Slate RHI Error - Invalid Texture2D, ResourceName: '%s'. UObject not valid."), *ResourceName.ToString());
+					}
 				}
 				ResourceObject = GetBadResourceTexture();
 #else
