[Integrate from Main - CL 45729864]

Removes default Post Process Buffer update - Post Process Buffers now only updated by [REDACTED_PERSON] widgets.
Also disables and deprecates the option for Post Update widgets to trigger full-screen updates rather than their paint area only.

#tests Tested in Editor and PIE
#changelist validated
#virtualized

diff --git a/Engine/Source/Runtime/SlateCore/Private/Rendering/ElementBatcher.cpp b/Engine/Source/Runtime/SlateCore/Private/Rendering/ElementBatcher.cpp
--- a/Engine/Source/Runtime/SlateCore/Private/Rendering/ElementBatcher.cpp
+++ b/Engine/Source/Runtime/SlateCore/Private/Rendering/ElementBatcher.cpp
@@ -293,7 +293,6 @@ FSlateElementBatcher::FSlateElementBatcher( TSharedRef<FSlateRenderingPolicy> In
 	, bCompositeHDRViewports(false)
 	, UsedSlatePostBuffers(ESlatePostRT::None)
 	, ResourceUpdatingPostBuffers(ESlatePostRT::None)
-	, SkipDefaultUpdatePostBuffers(ESlatePostRT::None)
 {
 }
 
@@ -3965,6 +3964,5 @@ void FSlateElementBatcher::ResetBatches()
 	bCompositeHDRViewports = false;
 	UsedSlatePostBuffers = ESlatePostRT::None;
 	ResourceUpdatingPostBuffers = ESlatePostRT::None;
-	SkipDefaultUpdatePostBuffers = ESlatePostRT::None;
 	NumPostProcessPasses = 0;
 }
diff --git a/Engine/Source/Runtime/SlateCore/Public/Rendering/ElementBatcher.h b/Engine/Source/Runtime/SlateCore/Public/Rendering/ElementBatcher.h
--- a/Engine/Source/Runtime/SlateCore/Public/Rendering/ElementBatcher.h
+++ b/Engine/Source/Runtime/SlateCore/Public/Rendering/ElementBatcher.h
@@ -270,12 +270,8 @@ public:
 
 	ESlatePostRT GetResourceUpdatingPostBuffers() const { return ResourceUpdatingPostBuffers; }
 
-	ESlatePostRT GetSkipDefaultUpdatePostBuffers() const { return SkipDefaultUpdatePostBuffers; }
-
 	void SetResourceUpdatingPostBuffers(ESlatePostRT InResourceUpdatingPostBuffers) { ResourceUpdatingPostBuffers = InResourceUpdatingPostBuffers; }
 
-	void SetSkipDefaultUpdatePostBuffers(ESlatePostRT InSkipDefaultUpdatePostBuffers) { SkipDefaultUpdatePostBuffers = InSkipDefaultUpdatePostBuffers; }
-
 	void SetCompositeHDRViewports(bool bInCompositeHDRViewports) { bCompositeHDRViewports = bInCompositeHDRViewports; }
 
 	/** 
@@ -527,7 +523,4 @@ private:
 
 	// true if a resource is updating a slate post buffer, if true we need to add a fence for the scene draw to complete before clearing unused post buffers
 	ESlatePostRT ResourceUpdatingPostBuffers;
-
-	// true if we should skip the default population of the post buffers with the scene & no UI.
-	ESlatePostRT SkipDefaultUpdatePostBuffers;
 };
diff --git a/Engine/Source/Runtime/SlateRHIRenderer/Private/SlateRHIRenderer.cpp b/Engine/Source/Runtime/SlateRHIRenderer/Private/SlateRHIRenderer.cpp
--- a/Engine/Source/Runtime/SlateRHIRenderer/Private/SlateRHIRenderer.cpp
+++ b/Engine/Source/Runtime/SlateRHIRenderer/Private/SlateRHIRenderer.cpp
@@ -181,20 +181,12 @@ struct FSlateViewportInfo : public FRenderResource
 	}
 };
 
-struct FSlatePostProcessUpdateRequest
-{
-	TSharedPtr<FSlateRHIPostBufferProcessorProxy> PostProcessorProxy;
-	FTextureResource* RenderTargetTextureResource = nullptr;
-	ESlatePostRT RenderTarget = ESlatePostRT::None;
-};
-
 struct FSlateDrawWindowPassInputs
 {
 	FSlateRHIRenderer* Renderer = nullptr;
 	FSlateWindowElementList* WindowElementList = nullptr;
 	SWindow* Window = nullptr;
 	FSlateViewportInfo* ViewportInfo = nullptr;
-	TConstArrayView<FSlatePostProcessUpdateRequest> PostProcessUpdateRequests;
 	FIntPoint CursorPosition = FIntPoint::ZeroValue;
 	FIntRect SceneViewRect;
 	float ViewportScaleUI = 0.0f;
@@ -710,48 +702,6 @@ FSlateDrawWindowPassOutputs FSlateRHIRenderer::DrawWindow_RenderThread(FRDGBuild
 		FRDGTexture* ElementsTexture = OutputTexture;
 		const FIntPoint OutputExtent = OutputTexture->Desc.Extent;
 
-		TArray<FRDGTexture*, FRDGArrayAllocator> PostProcessTextures;
-		PostProcessTextures.Reserve(Inputs.PostProcessUpdateRequests.Num());
-
-		for (const FSlatePostProcessUpdateRequest& Request : Inputs.PostProcessUpdateRequests)
-		{
-			FRDGTexture* Texture = RegisterExternalTexture(GraphBuilder, Request.RenderTargetTextureResource->GetTexture2DRHI(), TEXT("PostProcessRT"));
-			PostProcessTextures.Emplace(Texture);
-			GraphBuilder.UseInternalAccessMode(Texture);
-		}
-
-		// The post process input texture will be the separate viewport texture if it exists, or the swap chain.
-		FScreenPassTexture PostProcessInputTexture(RegisterExternalTexture(GraphBuilder, ViewportTextureRHI ? ViewportTextureRHI : SwapChainTextureRHI, TEXT("ViewportTexture")));
-
-		for (int32 PostProcessIndex = 0; PostProcessIndex < Inputs.PostProcessUpdateRequests.Num(); ++PostProcessIndex)
-		{
-			const FSlatePostProcessUpdateRequest& Request = Inputs.PostProcessUpdateRequests[PostProcessIndex];
-			const FScreenPassTexture PostProcessOutputTexture(PostProcessTextures[PostProcessIndex]);
-
-			if (ViewportTextureRHI)
-			{
-				PostProcessInputTexture.ViewRect = FIntRect(FIntPoint::ZeroValue, ViewportTextureRHI->GetSizeXY());
-			}
-			else
-			{
-				PostProcessInputTexture.ViewRect = Inputs.SceneViewRect;
-			}
-
-			if (Request.PostProcessorProxy)
-			{
-				Request.PostProcessorProxy->PostProcess_Renderthread(GraphBuilder, PostProcessInputTexture, PostProcessOutputTexture);
-			}
-			else
-			{
-				AddDrawTexturePass(GraphBuilder, FScreenPassViewInfo(), PostProcessInputTexture, PostProcessOutputTexture);
-			}
-		}
-
-		for (FRDGTexture* Texture : PostProcessTextures)
-		{
-			GraphBuilder.UseExternalAccessMode(Texture, ERHIAccess::SRVMask);
-		}
-
 		const bool bCompositeUIWithSceneHDR = ViewportInfo.bDisplayFormatIsHDR && CompositeUIWithSceneHDR();
 
 		bool bClearElementsTexture = Inputs.bClear || GSlateWireframe;
@@ -1110,7 +1060,6 @@ struct FSlateDrawWindowsCommand : public TConcurrentLinearObject<FSlateDrawWindo
 	bool IsEmpty() { return Windows.IsEmpty() && DeferredUpdates.IsEmpty(); }
 
 	TArray<FSlateDrawWindowPassInputs, FConcurrentLinearArrayAllocator> Windows;
-	TArray<FSlatePostProcessUpdateRequest, FConcurrentLinearArrayAllocator> PostProcessUpdates;
 	TArray<FRenderThreadUpdateContext, FConcurrentLinearArrayAllocator> DeferredUpdates;
 };
 
@@ -1210,7 +1159,6 @@ void FSlateRHIRenderer::DrawWindows_Private(FSlateDrawBuffer& WindowDrawBuffer)
 		float ViewportScaleUI;
 		ESlatePostRT PostProcessUsedBits;
 		ESlatePostRT PostProcessCustumDrawBits;
-		ESlatePostRT PostProcessSkipUpdateBits;
 		FIntPoint CursorPosition;
 		bool bLockToVsync;
 	};
@@ -1292,7 +1240,6 @@ void FSlateRHIRenderer::DrawWindows_Private(FSlateDrawBuffer& WindowDrawBuffer)
 				const bool bWindowCanRenderPostProcess       = bAppCanRender && ViewportExtent != FIntPoint::ZeroValue;
 				const ESlatePostRT PostProcessUsedBits       = bWindowCanRenderPostProcess ? ElementBatcher->GetUsedSlatePostBuffers() : ESlatePostRT::None;
 				const ESlatePostRT PostProcessCustomDrawBits = bWindowCanRenderPostProcess ? ElementBatcher->GetResourceUpdatingPostBuffers() : ESlatePostRT::None;
-				const ESlatePostRT PostProcessSkipUpdateBits = bWindowCanRenderPostProcess ? ElementBatcher->GetSkipDefaultUpdatePostBuffers() : ESlatePostRT::None;
 				const bool bLockToVsync                      = IsVSyncRequired(*ElementBatcher);
 
 				ElementBatcher->ResetBatches();
@@ -1309,7 +1256,6 @@ void FSlateRHIRenderer::DrawWindows_Private(FSlateDrawBuffer& WindowDrawBuffer)
 					, .ViewportScaleUI           = ViewportScaleUI
 					, .PostProcessUsedBits       = PostProcessUsedBits
 					, .PostProcessCustumDrawBits = PostProcessCustomDrawBits
-					, .PostProcessSkipUpdateBits = PostProcessSkipUpdateBits
 					, .CursorPosition            = FIntPoint(ViewportCursorPosition.X, ViewportCursorPosition.Y)
 					, .bLockToVsync              = bLockToVsync
 				});
@@ -1361,45 +1307,11 @@ void FSlateRHIRenderer::DrawWindows_Private(FSlateDrawBuffer& WindowDrawBuffer)
 	TUniquePtr<FSlateDrawWindowsCommand> DrawWindowsCommand = MakeUnique<FSlateDrawWindowsCommand>();
 	DrawWindowsCommand->Windows.Reserve(WindowsToRender.Num());
 	DrawWindowsCommand->DeferredUpdates = MoveTemp(DeferredUpdateContexts);
-	DrawWindowsCommand->PostProcessUpdates.Reserve(WindowsToRender.Num() * FMath::CountBits((int32)PostProcessAnyUsedBits));
 
 	bool bScreenshotProcessed = false;
 
 	for (const FWindowToRender& WindowToRender : WindowsToRender)
 	{
-		const int32 PostProcessUpdatesOffset = DrawWindowsCommand->PostProcessUpdates.Num();
-
-		if (bAppCanRenderPostProcess)
-		{
-			// Process bits that were NOT marked to skip the update.
-			for (ESlatePostRT Bit : MakeFlagsRange(ESlatePostRT::All & WindowToRender.PostProcessUsedBits & ~WindowToRender.PostProcessSkipUpdateBits))
-			{
-				UTextureRenderTarget2D* RenderTarget = RendererSettings->TryGetPostBufferRT(Bit);
-				check(RenderTarget);
-
-				FSlatePostProcessUpdateRequest Request;
-				Request.RenderTarget = Bit;
-				Request.RenderTargetTextureResource = RenderTarget->GetResource();
-
-				if (USlateRHIPostBufferProcessor* PostProcessor = USlateFXSubsystem::GetPostProcessor(Bit))
-				{
-					Request.PostProcessorProxy = PostProcessor->GetRenderThreadProxy();
-				}
-
-				checkf(DrawWindowsCommand->PostProcessUpdates.Num() != DrawWindowsCommand->PostProcessUpdates.Max(), TEXT("This container is about to resize which will result in a dangling memory access"));
-				DrawWindowsCommand->PostProcessUpdates.Add(Request);
-			}
-		}
-
-		TConstArrayView<FSlatePostProcessUpdateRequest> PostProcessUpdatesForWindow;
-
-		const int32 PostProcessUpdateRequestsCount = DrawWindowsCommand->PostProcessUpdates.Num() - PostProcessUpdatesOffset;
-
-		if (PostProcessUpdateRequestsCount != 0)
-		{
-			PostProcessUpdatesForWindow = MakeArrayView(&DrawWindowsCommand->PostProcessUpdates[PostProcessUpdatesOffset], PostProcessUpdateRequestsCount);
-		}
-
 		bScreenshotProcessed |= ScreenshotState.ViewportToCapture == WindowToRender.ViewportInfo;
 
 		if (bAppCanRender)
@@ -1410,7 +1322,6 @@ void FSlateRHIRenderer::DrawWindows_Private(FSlateDrawBuffer& WindowDrawBuffer)
 				, .WindowElementList         = WindowToRender.WindowElementList
 				, .Window                    = WindowToRender.Window
 				, .ViewportInfo              = WindowToRender.ViewportInfo
-				, .PostProcessUpdateRequests = PostProcessUpdatesForWindow
 				, .CursorPosition            = WindowToRender.CursorPosition
 				, .SceneViewRect             = WindowToRender.ViewportRect
 				, .ViewportScaleUI           = WindowToRender.ViewportScaleUI
diff --git a/Engine/Source/Runtime/UMG/Private/Components/PostBufferUpdate.cpp b/Engine/Source/Runtime/UMG/Private/Components/PostBufferUpdate.cpp
--- a/Engine/Source/Runtime/UMG/Private/Components/PostBufferUpdate.cpp
+++ b/Engine/Source/Runtime/UMG/Private/Components/PostBufferUpdate.cpp
@@ -13,31 +13,18 @@
 
 UPostBufferUpdate::UPostBufferUpdate()
 {
-	bPerformDefaultPostBufferUpdate = true;
-	bUpdateOnlyPaintArea = false;
 PRAGMA_DISABLE_DEPRECATION_WARNINGS
+	bUpdateOnlyPaintArea = true;
+	bPerformDefaultPostBufferUpdate = false;
 	BuffersToUpdate = {};
 PRAGMA_ENABLE_DEPRECATION_WARNINGS
 	UpdateBufferInfos = {};
 }
 
-void UPostBufferUpdate::SetPerformDefaultPostBufferUpdate(bool bInPerformDefaultPostBufferUpdate)
-{
-	if (bPerformDefaultPostBufferUpdate != bInPerformDefaultPostBufferUpdate)
-	{
-		bPerformDefaultPostBufferUpdate = bInPerformDefaultPostBufferUpdate;
-		if (MyPostBufferUpdate)
-		{
-			MyPostBufferUpdate->SetPerformDefaultPostBufferUpdate(bPerformDefaultPostBufferUpdate);
-		}
-	}
-}
-
 TSharedRef<SWidget> UPostBufferUpdate::RebuildWidget()
 {
 	MyPostBufferUpdate = SNew(SPostBufferUpdate)
-		.bUsePaintGeometry(bUpdateOnlyPaintArea)
-		.bPerformDefaultPostBufferUpdate(bPerformDefaultPostBufferUpdate);
+		.bUsePaintGeometry(true);
 
 	bool bSetBuffersToUpdate = true;
 
@@ -62,7 +49,6 @@ PRAGMA_DISABLE_DEPRECATION_WARNINGS
 				if (TObjectPtr<USlatePostBufferProcessorUpdater> PostParamUpdater = UpdateBufferInfo.PostParamUpdater)
 				{
 					TSharedPtr<FSlatePostProcessorUpdaterProxy> PostParamUpdaterProxy = PostParamUpdater->GetRenderThreadProxy();
-					PostParamUpdaterProxy->bSkipBufferUpdate = PostParamUpdater->bSkipBufferUpdate;
 					TempProcessorUpdaters.Add(UpdateBufferInfo.BufferToUpdate, PostParamUpdaterProxy);
 				}
 			}
@@ -89,8 +75,7 @@ void UPostBufferUpdate::SynchronizeProperties()
 		return;
 	}
 
-	MyPostBufferUpdate->SetPerformDefaultPostBufferUpdate(bPerformDefaultPostBufferUpdate);
-	MyPostBufferUpdate->SetUsePaintGeometry(bUpdateOnlyPaintArea);
+	MyPostBufferUpdate->SetUsePaintGeometry(true);
 
 	bool bSetBuffersToUpdate = true;
 
@@ -115,7 +100,6 @@ PRAGMA_DISABLE_DEPRECATION_WARNINGS
 				if (TObjectPtr<USlatePostBufferProcessorUpdater> PostParamUpdater = UpdateBufferInfo.PostParamUpdater)
 				{
 					TSharedPtr<FSlatePostProcessorUpdaterProxy> PostParamUpdaterProxy = PostParamUpdater->GetRenderThreadProxy();
-					PostParamUpdaterProxy->bSkipBufferUpdate = PostParamUpdater->bSkipBufferUpdate;
 					TempProcessorUpdaters.Add(UpdateBufferInfo.BufferToUpdate, PostParamUpdaterProxy);
 				}
 			}
diff --git a/Engine/Source/Runtime/UMG/Private/Slate/SPostBufferUpdate.cpp b/Engine/Source/Runtime/UMG/Private/Slate/SPostBufferUpdate.cpp
--- a/Engine/Source/Runtime/UMG/Private/Slate/SPostBufferUpdate.cpp
+++ b/Engine/Source/Runtime/UMG/Private/Slate/SPostBufferUpdate.cpp
@@ -36,9 +36,6 @@ public:
 
 public:
 
-	/** True if we should perform the default post buffer update, used to set related state on ElementBatcher at Gamethread element batch time */
-	bool bPerformDefaultPostBufferUpdate = true;
-
 	/** True if buffers to update has been initialized. Used to ensure 'BuffersToUpdate_Renderthread' is only set once at initialization. */
 	bool bBuffersToUpdateInitialized = false;
 
@@ -104,11 +101,6 @@ void FPostBufferUpdater::Draw_RenderThread(FRDGBuilder& GraphBuilder, const FDra
 				if (TSharedPtr<FSlatePostProcessorUpdaterProxy> ProcessorUpdater = *ProcessorUpdaterItr)
 				{
 					ProcessorUpdater->UpdateProcessor_RenderThread(PostProcessorProxy);
-
-					if (ProcessorUpdater->bSkipBufferUpdate)
-					{
-						continue;
-					}
 				}
 			}
 		}
@@ -174,12 +166,6 @@ void FPostBufferUpdater::PostCustomElementAdded(FSlateElementBatcher& ElementBat
 	ESlatePostRT PrevResourceUpdatingPostBuffers = ElementBatcher.GetResourceUpdatingPostBuffers();
 	ElementBatcher.SetResourceUpdatingPostBuffers(PrevResourceUpdatingPostBuffers | BuffersToUpdate_Renderthread);
 
-	if (!bPerformDefaultPostBufferUpdate)
-	{
-		ESlatePostRT PrevSkipDefaultUpdatePostBuffers = ElementBatcher.GetSkipDefaultUpdatePostBuffers();
-		ElementBatcher.SetSkipDefaultUpdatePostBuffers(PrevSkipDefaultUpdatePostBuffers | BuffersToUpdate_Renderthread);
-	}
-
 	// Give proxies a chance to update their renderthread values.
 	if (const USlateRHIRendererSettings* RendererSettings = USlateRHIRendererSettings::Get())
 	{
@@ -210,7 +196,6 @@ void SPostBufferUpdate::PrivateRegisterAttributes(FSlateAttributeInitializer& At
 
 SPostBufferUpdate::SPostBufferUpdate()
 	: bUsePaintGeometry(false)
-	, bPerformDefaultPostBufferUpdate(true)
 	, BuffersToUpdate({})
 	, PostBufferUpdater(nullptr)
 {
@@ -220,15 +205,12 @@ void SPostBufferUpdate::Construct(const FArguments& InArgs)
 {
 #if !UE_SERVER
 	bUsePaintGeometry = InArgs._bUsePaintGeometry;
-	bPerformDefaultPostBufferUpdate = InArgs._bPerformDefaultPostBufferUpdate;
 
 	BuffersToUpdate = {};
 
 	PostBufferUpdater = MakeShared<FPostBufferUpdater>();
 	if (PostBufferUpdater)
 	{
-		PostBufferUpdater->bPerformDefaultPostBufferUpdate = bPerformDefaultPostBufferUpdate;
-
 		// Safe to set RT values here on construct
 		PostBufferUpdater->bUsePaintGeometry_Renderthread = bUsePaintGeometry;
 	}
@@ -256,23 +238,6 @@ void SPostBufferUpdate::SetUsePaintGeometry(bool bInUsePaintGeometry)
 #endif // !UE_SERVER
 }
 
-void SPostBufferUpdate::SetPerformDefaultPostBufferUpdate(bool bInPerformDefaultPostBufferUpdate)
-{
-#if !UE_SERVER
-	bPerformDefaultPostBufferUpdate = bInPerformDefaultPostBufferUpdate;
-
-	if (PostBufferUpdater)
-	{
-		PostBufferUpdater->bPerformDefaultPostBufferUpdate = bPerformDefaultPostBufferUpdate;
-	}
-#endif // !UE_SERVER
-}
-
-bool SPostBufferUpdate::GetPerformDefaultPostBufferUpdate() const
-{
-	return bPerformDefaultPostBufferUpdate;
-}
-
 void SPostBufferUpdate::SetBuffersToUpdate(const TArrayView<ESlatePostRT> InBuffersToUpdate)
 {
 #if !UE_SERVER
diff --git a/Engine/Source/Runtime/UMG/Public/Components/PostBufferUpdate.h b/Engine/Source/Runtime/UMG/Public/Components/PostBufferUpdate.h
--- a/Engine/Source/Runtime/UMG/Public/Components/PostBufferUpdate.h
+++ b/Engine/Source/Runtime/UMG/Public/Components/PostBufferUpdate.h
@@ -50,17 +50,19 @@ public:
 
 private:
 
+	UE_DEPRECATED(5.7, "Now always true. The value of bUpdateOnlyPaintArea is ignored")
 	/**
 	 * True if we should only update the buffer within the bounds of this widget
 	 */
-	UPROPERTY(EditAnywhere, BlueprintReadOnly, Category = "Behavior", meta = (AllowPrivateAccess = "true"))
+	UPROPERTY(BlueprintReadOnly, Category = "Behavior", meta = (AllowPrivateAccess = "true"))
 	bool bUpdateOnlyPaintArea;
 
-	/** 
+	UE_DEPRECATED(5.7, "Default post buffer updates are no longer performed under any circumstances")
+	/**
 	 * True if we should do the default post buffer update of the scene before any UI.
 	 * If any PostBufferUpdate widget has this set as false, be default scene copy / processing will not occur.
 	 */
-	UPROPERTY(EditAnywhere, BlueprintReadWrite, Setter="SetPerformDefaultPostBufferUpdate", Category = "Behavior", meta = (AllowPrivateAccess = "true"))
+	UPROPERTY(BlueprintReadOnly, Category = "Behavior", meta = (AllowPrivateAccess = "true"))
 	bool bPerformDefaultPostBufferUpdate;
 
 	UE_DEPRECATED(5.5, "BuffersToUpdate is deprecated. Please use UpdateBufferInfos. This array will be ignored if UpdateBufferInfos is used")
@@ -73,11 +75,9 @@ private:
 	TArray<FSlatePostBufferUpdateInfo> UpdateBufferInfos;
 
 public:
-	/** Set the orientation of the stack box. The existing elements will be rearranged. */
-	UMG_API void SetPerformDefaultPostBufferUpdate(bool bInPerformDefaultPostBufferUpdate);
 
 #if WITH_EDITOR
-	bool GetUpdateOnlyPaintArea() const { return bUpdateOnlyPaintArea; }
+	bool GetUpdateOnlyPaintArea() const { return true; }
 
 	/** True if this widget updates any buffer indicated in SlateBufferBits */
 	UMG_API bool DoesUpdateSlatePostRT(ESlatePostRT SlateBufferBits) const ;
@@ -106,11 +106,15 @@ class USlatePostBufferProcessorUpdater : public UObject
 
 public:
 
+	UE_DEPRECATED(5.7, "Now always false. The value of bSkipBufferUpdate is ignored")
 	/**
 	 * True implies we will skip the buffer update & only update the processor.
-	 * Useful to reset params for processor runs next frame
+	 * Previously useful to reset params for processor runs next frame.
+	 * Since there are no default buffer updates anymore, setting parameters
+	 * for subsequent updates is meaningless as non-default updates will set
+	 * all their parameters anyway.
 	 */
-	UPROPERTY(EditAnywhere, Category = "Update")
+	UPROPERTY()
 	bool bSkipBufferUpdate = false;
 
 public:
diff --git a/Engine/Source/Runtime/UMG/Public/Slate/SPostBufferUpdate.h b/Engine/Source/Runtime/UMG/Public/Slate/SPostBufferUpdate.h
--- a/Engine/Source/Runtime/UMG/Public/Slate/SPostBufferUpdate.h
+++ b/Engine/Source/Runtime/UMG/Public/Slate/SPostBufferUpdate.h
@@ -29,12 +29,6 @@ class FSlatePostProcessorUpdaterProxy : public TSharedFromThis<FSlatePostProcess
 {
 public:
 
-	/**
-	 * True implies we will skip the buffer update & only update the processor.
-	 * Useful to reset params for processor runs next frame
-	 */
-	bool bSkipBufferUpdate = false;
-
 public:
 	
 	virtual ~FSlatePostProcessorUpdaterProxy ()
@@ -55,14 +49,10 @@ class SPostBufferUpdate
 public:
 	SLATE_BEGIN_ARGS( SPostBufferUpdate )
 		: _bUsePaintGeometry(false)
-		, _bPerformDefaultPostBufferUpdate(true)
 		{ }
 
 		/** True if we should perform the default post buffer update */
 		SLATE_ARGUMENT(bool, bUsePaintGeometry)
-
-		/** True if we should perform the default post buffer update */
-		SLATE_ARGUMENT(bool, bPerformDefaultPostBufferUpdate)
 	SLATE_END_ARGS()
 
 	/** Constructor */
@@ -79,12 +69,6 @@ public:
 	/** Set if we should use paint geometry */
 	UMG_API void SetUsePaintGeometry(bool bInUsePaintGeometry);
 
-	/** Set if we should perform the default post buffer update */
-	UMG_API void SetPerformDefaultPostBufferUpdate(bool bInPerformDefaultPostBufferUpdate);
-
-	/** Get if we should perform the default post buffer update */
-	UMG_API bool GetPerformDefaultPostBufferUpdate() const;
-
 	/**
 	 * Set buffers to, this method only affects the 'FPostBufferUpdater' custom drawer once during initialization. 
 	 * Afterwards it will be a no-op due to the custom drawer having update buffers initialized.
@@ -121,9 +105,6 @@ private:
 	/** True if we should limit processing to our paint geometry */
 	bool bUsePaintGeometry;
 
-	/** True if we should perform the default post buffer update */
-	bool bPerformDefaultPostBufferUpdate;
-
 	/** Buffers that we should update, all of these buffers will be affected by 'bPerformDefaultPostBufferUpdate' if disabled */
 	TArray<ESlatePostRT> BuffersToUpdate;
 
