Add project setting to switch Slate material shaders between fully static specializations, hybrid and fully dynamic. This allows for a balance between GPU performance and the number of shaders compiled / PSOs created.

diff --git a/Engine/Source/Runtime/Engine/Classes/Engine/RendererSettings.h b/Engine/Source/Runtime/Engine/Classes/Engine/RendererSettings.h
--- a/Engine/Source/Runtime/Engine/Classes/Engine/RendererSettings.h
+++ b/Engine/Source/Runtime/Engine/Classes/Engine/RendererSettings.h
@@ -301,6 +301,17 @@ namespace EWorkingColorSpace
 	};
 }
 
+UENUM()
+enum class ESlateMaterialShaderMode : uint8
+{
+	/* Use static specializations for all Slate material shaders. This results in the most shaders and PSOs compiled for each Slate material, but leads to optimal performance. */
+	Static = 0		UMETA(DisplayName = "Static"),
+	/* Use static specialization for Default Slate shader type and dynamic branch version for the rest. This balances performance and PSO/shader count. */
+	Hybrid  = 1 	UMETA(DisplayName = "Hybrid"),
+	/* Only use dynamic branch version of Slate material shaders. Leads to some performance overhead, but the least number of total shaders. */
+	Dynamic = 2 	UMETA(DisplayName = "Dynamic"),
+};
+
 /**
  * Rendering settings.
  */
@@ -1439,6 +1450,13 @@ class URendererSettings : public UDeveloperSettings
 		ConfigRestartRequired = true))
 		uint32 bSupportReversedIndexBuffers : 1;
 
+	UPROPERTY(config, EditAnywhere, Category = Optimizations, meta = (
+		ConsoleVariable = "slate.MaterialShaderMode", DisplayName = "Slate Material Shader Mode",
+		ToolTip = "Specify how Slate material shaders should be compiled to balance GPU performance and number of shaders",
+		ConfigRestartRequired = true))
+		ESlateMaterialShaderMode SlateMaterialShaderMode;
+
+
 	UPROPERTY(config, EditAnywhere, Category = Mobile, meta = (
 		ConsoleVariable = "r.Mobile.AmbientOcclusion", DisplayName = "Ambient Occlusion (When using ForwardShading)",
 		ToolTip = "Mobile Ambient Occlusion. Causion: An extra sampler will be occupied in mobile base pass pixel shader after enable the mobile ambient occlusion. Changing this setting requires restarting the editor.",
diff --git a/Engine/Source/Runtime/SlateRHIRenderer/Private/SlateMaterialShader.h b/Engine/Source/Runtime/SlateRHIRenderer/Private/SlateMaterialShader.h
--- a/Engine/Source/Runtime/SlateRHIRenderer/Private/SlateMaterialShader.h
+++ b/Engine/Source/Runtime/SlateRHIRenderer/Private/SlateMaterialShader.h
@@ -95,6 +95,8 @@ public:
 	}
 };
 
+bool IsValidSlateMaterialShaderType(ESlateShader ShaderType);
+
 template<ESlateShader ShaderType>
 class TSlateMaterialShaderPS : public FSlateMaterialShaderPS
 {
@@ -110,7 +112,7 @@ public:
 	/** Only compile shaders used with UI. */
 	static bool ShouldCompilePermutation(const FMaterialShaderPermutationParameters& Parameters)
 	{
-		return FSlateMaterialShaderPS::ShouldCompilePermutation(Parameters);
+		return FSlateMaterialShaderPS::ShouldCompilePermutation(Parameters) && IsValidSlateMaterialShaderType(ShaderType);
 	}
 
 	/** Modifies the compilation of this shader. */
diff --git a/Engine/Source/Runtime/SlateRHIRenderer/Private/SlateRHIRenderingPolicy.cpp b/Engine/Source/Runtime/SlateRHIRenderer/Private/SlateRHIRenderingPolicy.cpp
--- a/Engine/Source/Runtime/SlateRHIRenderer/Private/SlateRHIRenderingPolicy.cpp
+++ b/Engine/Source/Runtime/SlateRHIRenderer/Private/SlateRHIRenderingPolicy.cpp
@@ -57,13 +57,31 @@ static FAutoConsoleVariableRef CVarGSlateMaterialPSOPrecache(
 	TEXT("Precache all possible required PSOs for loaded Slate Materials."),
 	ECVF_ReadOnly);
 
-static int32 GSlateDynamicShaders = 0;
-static FAutoConsoleVariableRef CVarGSlateDynamicShaders(
-	TEXT("r.Slate.DynamicShaders"),
-	GSlateDynamicShaders,
-	TEXT("Whether to use a dynamic branch path for all Slate materials instead of static specializations. ")
-	TEXT("This reduces the number of unique Slate PSOs, but may result in some GPU overhead. (Default: 0)"),
-	ECVF_Default);
+static int32 GSlateMaterialShaderMode = 0;
+static FAutoConsoleVariableRef CVarSlateMaterialShaderMode(
+	TEXT("slate.MaterialShaderMode"),
+	GSlateMaterialShaderMode,
+	TEXT("Whether to compile fully specialized slate material shaders or use dynamic branches. Changing this variable affects which shaders are compiled during a cook and are available at runtime.")
+	TEXT(" 0: Fully static specialized shaders (best GPU performance, default)\n")
+	TEXT(" 1: Hybrid approach, with specialization for Default shader type and dynamic branch for the rest (balance between performance and shader count)\n")
+	TEXT(" 2: Only use dynamic branch version of the shaders (lowest shader count)\n"),
+	ECVF_ReadOnly);
+
+bool IsValidSlateMaterialShaderType(ESlateShader ShaderType)
+{
+	switch (ESlateMaterialShaderMode(GSlateMaterialShaderMode))
+	{
+	case ESlateMaterialShaderMode::Static:
+		return ShaderType != ESlateShader::Dynamic;
+	case ESlateMaterialShaderMode::Hybrid:
+		return  ShaderType == ESlateShader::Default
+			 || ShaderType == ESlateShader::Dynamic;
+	case ESlateMaterialShaderMode::Dynamic:
+		return ShaderType == ESlateShader::Dynamic;
+	default:
+		return true;
+	}
+}
 
 static const TCHAR* SlateGlobalPSOCollectorName = TEXT("SlateGlobalPSOCollector");
 static const TCHAR* SlateMaterialPSOCollectorName = TEXT("SlateMaterialPSOCollector");
@@ -439,12 +457,9 @@ TShaderRef<FSlateElementPS> GetTexturePixelShader(FGlobalShaderMap* ShaderMap, E
 
 bool ChooseMaterialShaderTypes(ESlateShader ShaderType, bool bUseInstancing, FMaterialShaderTypes& OutShaderTypes)
 {
-	if (GSlateDynamicShaders)
-	{
-		OutShaderTypes.AddShaderType<TSlateMaterialShaderPS<ESlateShader::Dynamic>>();
-	}
-	else
+	switch (ESlateMaterialShaderMode(GSlateMaterialShaderMode))
 	{
+	case ESlateMaterialShaderMode::Static:
 		switch (ShaderType)
 		{
 		case ESlateShader::Default:
@@ -474,6 +489,22 @@ bool ChooseMaterialShaderTypes(ESlateShader ShaderType, bool bUseInstancing, FMa
 		default:
 			return false;
 		}
+		break;
+	case ESlateMaterialShaderMode::Hybrid:
+		if (ShaderType == ESlateShader::Default)
+		{
+			OutShaderTypes.AddShaderType<TSlateMaterialShaderPS<ESlateShader::Default>>();
+		}
+		else
+		{
+			OutShaderTypes.AddShaderType<TSlateMaterialShaderPS<ESlateShader::Dynamic>>();
+		}
+		break;
+	case ESlateMaterialShaderMode::Dynamic:
+		OutShaderTypes.AddShaderType<TSlateMaterialShaderPS<ESlateShader::Dynamic>>();
+		break;
+	default:
+		return false;
 	}
 
 	if (bUseInstancing)
@@ -1725,11 +1756,6 @@ void AddSlateDrawElementsPass(
 
 TConstArrayView<ESlateShader> GetSlateShaderTypesToPrecache()
 {
-	static const ESlateShader DynamicShaderTypes[] =
-	{ 
-		ESlateShader::Dynamic 
-	};
-
 	static const ESlateShader StaticShaderTypes[] =
 	{
 		ESlateShader::Default,
@@ -1739,9 +1765,28 @@ TConstArrayView<ESlateShader> GetSlateShaderTypesToPrecache()
 		ESlateShader::RoundedBox
 	};
 
-	return GSlateDynamicShaders 
-		? TConstArrayView<ESlateShader>(DynamicShaderTypes)
-		: TConstArrayView<ESlateShader>(StaticShaderTypes);
+	static const ESlateShader HybridShaderTypes[] =
+	{
+		ESlateShader::Default,
+		ESlateShader::Dynamic
+	};
+
+	static const ESlateShader DynamicShaderTypes[] =
+	{
+		ESlateShader::Dynamic
+	};
+
+	switch (ESlateMaterialShaderMode(GSlateMaterialShaderMode))
+	{
+	case ESlateMaterialShaderMode::Static:
+		return TConstArrayView<ESlateShader>(StaticShaderTypes);
+	case ESlateMaterialShaderMode::Hybrid:
+		return TConstArrayView<ESlateShader>(HybridShaderTypes);
+	case ESlateMaterialShaderMode::Dynamic:
+		return TConstArrayView<ESlateShader>(DynamicShaderTypes);
+	default:
+		return TConstArrayView<ESlateShader>();
+	}
 }
 
 static EPixelFormat GetDefaultBackBufferPixelFormat()
