Only track slate PSO precache misses when Slate PSO precaching is enabled

diff --git a/Engine/Source/Runtime/SlateRHIRenderer/Private/SlateRHIRenderingPolicy.cpp b/Engine/Source/Runtime/SlateRHIRenderer/Private/SlateRHIRenderingPolicy.cpp
--- a/Engine/Source/Runtime/SlateRHIRenderer/Private/SlateRHIRenderingPolicy.cpp
+++ b/Engine/Source/Runtime/SlateRHIRenderer/Private/SlateRHIRenderingPolicy.cpp
@@ -1373,7 +1373,7 @@ void DrawSlateRenderBatch(
 		State.GraphicsPSOInit.PrimitiveType = GetRHIPrimitiveType(RenderBatch.DrawPrimitiveType);
 		
 #if PSO_PRECACHING_VALIDATE
-		if (PSOCollectorStats::IsFullPrecachingValidationEnabled())
+		if (PSOCollectorStats::IsFullPrecachingValidationEnabled() && GSlateMaterialPSOPrecache > 0)
 		{
 			static const int32 MaterialPSOCollectorIndex = FPSOCollectorCreateManager::GetIndex(GetFeatureLevelShadingPath(GMaxRHIFeatureLevel), SlateMaterialPSOCollectorName);
 			// Material Render Proxy is not cached in the state, but could be done for better debugging if needed
