Changing FSlateFontTextureRHI to not use implied command lists.

diff --git a/Engine/Source/Runtime/SlateRHIRenderer/Private/SlateRHIFontTexture.cpp b/Engine/Source/Runtime/SlateRHIRenderer/Private/SlateRHIFontTexture.cpp
--- a/Engine/Source/Runtime/SlateRHIRenderer/Private/SlateRHIFontTexture.cpp
+++ b/Engine/Source/Runtime/SlateRHIRenderer/Private/SlateRHIFontTexture.cpp
@@ -139,13 +139,14 @@ void FSlateFontAtlasRHI::ConditionalUpdateTexture()
 	{
 		if (IsInRenderingThread())
 		{
-			FontTexture->InitResource(FRHICommandListImmediate::Get());
+			FRHICommandListImmediate& RHICmdList = FRHICommandListImmediate::Get();
+			FontTexture->InitResource(RHICmdList);
 
 			uint32 DestStride;
-			uint8* TempData = (uint8*)RHILockTexture2D( FontTexture->GetTypedResource(), 0, RLM_WriteOnly, /*out*/ DestStride, false );
+			uint8* TempData = (uint8*)RHICmdList.LockTexture2D( FontTexture->GetTypedResource(), 0, RLM_WriteOnly, /*out*/ DestStride, false );
 			// check( DestStride == Atlas.BytesPerPixel * Atlas.AtlasWidth ); // Temporarily disabling check
 			FMemory::Memcpy( TempData, AtlasData.GetData(), GetSlateFontAtlasContentBytesPerPixel(ContentType)*AtlasWidth*AtlasHeight );
-			RHIUnlockTexture2D( FontTexture->GetTypedResource(),0,false );
+			RHICmdList.UnlockTexture2D( FontTexture->GetTypedResource(),0,false );
 		}
 		else
 		{
@@ -158,10 +159,10 @@ void FSlateFontAtlasRHI::ConditionalUpdateTexture()
 				[Atlas](FRHICommandListImmediate& RHICmdList)
 				{
 					uint32 DestStride;
-					uint8* TempData = (uint8*)RHILockTexture2D( Atlas->FontTexture->GetTypedResource(), 0, RLM_WriteOnly, /*out*/ DestStride, false );
+					uint8* TempData = (uint8*)RHICmdList.LockTexture2D( Atlas->FontTexture->GetTypedResource(), 0, RLM_WriteOnly, /*out*/ DestStride, false );
 					// check( DestStride == Atlas.BytesPerPixel * Atlas.AtlasWidth ); // Temporarily disabling check
 					FMemory::Memcpy( TempData, Atlas->AtlasData.GetData(), GetSlateFontAtlasContentBytesPerPixel(Atlas->GetContentType())*Atlas->AtlasWidth*Atlas->AtlasHeight );
-					RHIUnlockTexture2D( Atlas->FontTexture->GetTypedResource(),0,false );
+					RHICmdList.UnlockTexture2D( Atlas->FontTexture->GetTypedResource(),0,false );
 				});
 		}
 
@@ -174,8 +175,10 @@ FSlateFontTextureRHI::FSlateFontTextureRHI(const uint32 InWidth, const uint32 In
 {
 	if (IsInRenderingThread())
 	{
-		FontTexture->InitResource(FRHICommandListImmediate::Get());
-		UpdateTextureFromSource(InWidth, InHeight, InRawData);
+		FRHICommandListImmediate& RHICmdList = FRHICommandListImmediate::Get();
+
+		FontTexture->InitResource(RHICmdList);
+		UpdateTextureFromSource(RHICmdList, InWidth, InHeight, InRawData);
 	}
 	else
 	{
@@ -189,7 +192,7 @@ FSlateFontTextureRHI::FSlateFontTextureRHI(const uint32 InWidth, const uint32 In
 		ENQUEUE_RENDER_COMMAND(SlateUpdateFontTextureCommand)(
 			[This](FRHICommandListImmediate& RHICmdList)
 			{
-				This->UpdateTextureFromSource(This->PendingSourceData->SourceWidth, This->PendingSourceData->SourceHeight, This->PendingSourceData->SourceData);
+				This->UpdateTextureFromSource(RHICmdList, This->PendingSourceData->SourceWidth, This->PendingSourceData->SourceHeight, This->PendingSourceData->SourceData);
 				This->PendingSourceData.Reset();
 			});
 	}
@@ -212,12 +215,12 @@ void FSlateFontTextureRHI::ReleaseResources()
 	BeginReleaseResource(FontTexture.Get());
 }
 
-void FSlateFontTextureRHI::UpdateTextureFromSource(const uint32 SourceWidth, const uint32 SourceHeight, const TArray<uint8>& SourceData)
+void FSlateFontTextureRHI::UpdateTextureFromSource(FRHICommandListImmediate& RHICmdList, const uint32 SourceWidth, const uint32 SourceHeight, const TArray<uint8>& SourceData)
 {
 	const uint32 BytesPerPixel = GetSlateFontAtlasContentBytesPerPixel(GetContentType());
 
 	uint32 DestStride;
-	uint8* LockedTextureData = static_cast<uint8*>(RHILockTexture2D(FontTexture->GetTypedResource(), 0, RLM_WriteOnly, /*out*/ DestStride, false));
+	uint8* LockedTextureData = static_cast<uint8*>(RHICmdList.LockTexture2D(FontTexture->GetTypedResource(), 0, RLM_WriteOnly, /*out*/ DestStride, false));
 
 	// If our stride matches our width, we can copy in a single call, rather than copy each line
 	if (BytesPerPixel * SourceWidth == DestStride)
@@ -232,5 +235,5 @@ void FSlateFontTextureRHI::UpdateTextureFromSource(const uint32 SourceWidth, con
 		}
 	}
 	
-	RHIUnlockTexture2D(FontTexture->GetTypedResource(), 0, false);
+	RHICmdList.UnlockTexture2D(FontTexture->GetTypedResource(), 0, false);
 }
diff --git a/Engine/Source/Runtime/SlateRHIRenderer/Private/SlateRHIFontTexture.h b/Engine/Source/Runtime/SlateRHIRenderer/Private/SlateRHIFontTexture.h
--- a/Engine/Source/Runtime/SlateRHIRenderer/Private/SlateRHIFontTexture.h
+++ b/Engine/Source/Runtime/SlateRHIRenderer/Private/SlateRHIFontTexture.h
@@ -94,7 +94,7 @@ public:
 
 	void ReleaseResources();
 private:
-	void UpdateTextureFromSource(const uint32 SourceWidth, const uint32 SourceHeight, const TArray<uint8>& SourceData);
+	void UpdateTextureFromSource(FRHICommandListImmediate& RHICmdList, const uint32 SourceWidth, const uint32 SourceHeight, const TArray<uint8>& SourceData);
 private:
 	struct FPendingSourceData
 	{
