Change all slate.MaterialShaderMode access to per platform.
#rnx
#Synced-CL 48841205

diff --git a/Engine/Source/Runtime/Engine/Classes/Engine/RendererSettings.h b/Engine/Source/Runtime/Engine/Classes/Engine/RendererSettings.h
--- a/Engine/Source/Runtime/Engine/Classes/Engine/RendererSettings.h
+++ b/Engine/Source/Runtime/Engine/Classes/Engine/RendererSettings.h
@@ -1424,7 +1424,7 @@ class URendererSettings : public UDeveloperSettings
 
 	UPROPERTY(config, EditAnywhere, Category = Optimizations, meta = (
 		ConsoleVariable = "slate.MaterialShaderMode", DisplayName = "Slate Material Shader Mode",
-		ToolTip = "Specify how Slate material shaders should be compiled to balance GPU performance and number of shaders",
+		ToolTip = "Specify how Slate material shaders should be compiled to balance GPU performance and number of shaders. This can be overridden using per-platform configuration files.",
 		ConfigRestartRequired = true))
 		ESlateMaterialShaderMode SlateMaterialShaderMode;
 
diff --git a/Engine/Source/Runtime/SlateRHIRenderer/Private/SlateMaterialShader.h b/Engine/Source/Runtime/SlateRHIRenderer/Private/SlateMaterialShader.h
--- a/Engine/Source/Runtime/SlateRHIRenderer/Private/SlateMaterialShader.h
+++ b/Engine/Source/Runtime/SlateRHIRenderer/Private/SlateMaterialShader.h
@@ -95,7 +95,7 @@ public:
 	}
 };
 
-bool IsValidSlateMaterialShaderType(ESlateShader ShaderType);
+bool IsValidSlateMaterialShaderType(ESlateShader ShaderType, EShaderPlatform ShaderPlatform);
 
 template<ESlateShader ShaderType>
 class TSlateMaterialShaderPS : public FSlateMaterialShaderPS
@@ -112,7 +112,7 @@ public:
 	/** Only compile shaders used with UI. */
 	static bool ShouldCompilePermutation(const FMaterialShaderPermutationParameters& Parameters)
 	{
-		return FSlateMaterialShaderPS::ShouldCompilePermutation(Parameters) && IsValidSlateMaterialShaderType(ShaderType);
+		return FSlateMaterialShaderPS::ShouldCompilePermutation(Parameters) && IsValidSlateMaterialShaderType(ShaderType, Parameters.Platform);
 	}
 
 	/** Modifies the compilation of this shader. */
diff --git a/Engine/Source/Runtime/SlateRHIRenderer/Private/SlateRHIRenderingPolicy.cpp b/Engine/Source/Runtime/SlateRHIRenderer/Private/SlateRHIRenderingPolicy.cpp
--- a/Engine/Source/Runtime/SlateRHIRenderer/Private/SlateRHIRenderingPolicy.cpp
+++ b/Engine/Source/Runtime/SlateRHIRenderer/Private/SlateRHIRenderingPolicy.cpp
@@ -57,19 +57,25 @@ static FAutoConsoleVariableRef CVarGSlateMaterialPSOPrecache(
 	TEXT("Precache all possible required PSOs for loaded Slate Materials."),
 	ECVF_ReadOnly);
 
-static int32 GSlateMaterialShaderMode = 0;
-static FAutoConsoleVariableRef CVarSlateMaterialShaderMode(
+static TAutoConsoleVariable<int32> CVarSlateMaterialShaderMode(
 	TEXT("slate.MaterialShaderMode"),
-	GSlateMaterialShaderMode,
+	0,
 	TEXT("Whether to compile fully specialized slate material shaders or use dynamic branches. Changing this variable affects which shaders are compiled during a cook and are available at runtime.")
 	TEXT(" 0: Fully static specialized shaders (best GPU performance, default)\n")
 	TEXT(" 1: Hybrid approach, with specialization for Default shader type and dynamic branch for the rest (balance between performance and shader count)\n")
 	TEXT(" 2: Only use dynamic branch version of the shaders (lowest shader count)\n"),
 	ECVF_ReadOnly);
 
-bool IsValidSlateMaterialShaderType(ESlateShader ShaderType)
+static int32 GetSlateMaterialShaderMode(EShaderPlatform ShaderPlatform)
 {
-	switch (ESlateMaterialShaderMode(GSlateMaterialShaderMode))
+	static FShaderPlatformCachedIniValue<int32> PerPlatformCVar(TEXT("slate.MaterialShaderMode"));
+	return PerPlatformCVar.Get(ShaderPlatform);
+}
+
+bool IsValidSlateMaterialShaderType(ESlateShader ShaderType, EShaderPlatform ShaderPlatform)
+{
+	const int32 SlateMaterialShaderMode = GetSlateMaterialShaderMode(ShaderPlatform);
+	switch (ESlateMaterialShaderMode(SlateMaterialShaderMode))
 	{
 	case ESlateMaterialShaderMode::Static:
 		return ShaderType != ESlateShader::Dynamic;
@@ -455,9 +461,10 @@ TShaderRef<FSlateElementPS> GetTexturePixelShader(FGlobalShaderMap* ShaderMap, E
 	return PixelShader;
 }
 
-bool ChooseMaterialShaderTypes(ESlateShader ShaderType, bool bUseInstancing, FMaterialShaderTypes& OutShaderTypes)
+bool ChooseMaterialShaderTypes(ESlateShader ShaderType, EShaderPlatform ShaderPlatform, bool bUseInstancing, FMaterialShaderTypes& OutShaderTypes)
 {
-	switch (ESlateMaterialShaderMode(GSlateMaterialShaderMode))
+	const int32 SlateMaterialShaderMode = GetSlateMaterialShaderMode(ShaderPlatform);
+	switch (ESlateMaterialShaderMode(SlateMaterialShaderMode))
 	{
 	case ESlateMaterialShaderMode::Static:
 		switch (ShaderType)
@@ -1123,6 +1130,7 @@ FSlateRenderBatchOp* CreateSlateRenderBatchOp(
 
 		const FSlateSceneView* SceneView = Inputs.SceneViewAllocator->BeginAllocateSceneView(GraphBuilder, RenderBatch->SceneIndex);
 		const ERHIFeatureLevel::Type SceneFeatureLevel = SceneView->FeatureLevel;
+		const EShaderPlatform SceneShaderPlatform = GetFeatureLevelShaderPlatform(SceneFeatureLevel);
 		const FSceneInterface* Scene = SceneView->Scene;
 		const TUniformBufferRef<FViewUniformShaderParameters>& ViewUniformBuffer = Inputs.SceneViewAllocator->GetViewUniformBuffer(SceneView);
 
@@ -1130,7 +1138,7 @@ FSlateRenderBatchOp* CreateSlateRenderBatchOp(
 		TShaderRef<FSlateMaterialShaderPS> PixelShader;
 
 		FMaterialShaderTypes ShaderTypesToGet;
-		if (!ChooseMaterialShaderTypes(ShaderType, bUseInstancing, ShaderTypesToGet))
+		if (!ChooseMaterialShaderTypes(ShaderType, SceneShaderPlatform, bUseInstancing, ShaderTypesToGet))
 		{
 			checkf(false, TEXT("Unsupported Slate shader type for use with materials"));
 			return nullptr;
@@ -1755,7 +1763,7 @@ void AddSlateDrawElementsPass(
 	INC_DWORD_STAT_BY(STAT_SlateStencilClips, ClippingCreateContext.NumStencils);
 }
 
-TConstArrayView<ESlateShader> GetSlateShaderTypesToPrecache()
+TConstArrayView<ESlateShader> GetSlateShaderTypesToPrecache(EShaderPlatform ShaderPlatform)
 {
 	static const ESlateShader StaticShaderTypes[] =
 	{
@@ -1777,7 +1785,8 @@ TConstArrayView<ESlateShader> GetSlateShaderTypesToPrecache()
 		ESlateShader::Dynamic
 	};
 
-	switch (ESlateMaterialShaderMode(GSlateMaterialShaderMode))
+	const int32 SlateMaterialShaderMode = GetSlateMaterialShaderMode(ShaderPlatform);
+	switch (ESlateMaterialShaderMode(SlateMaterialShaderMode))
 	{
 	case ESlateMaterialShaderMode::Static:
 		return TConstArrayView<ESlateShader>(StaticShaderTypes);
@@ -1893,9 +1902,12 @@ FRegisterGlobalPSOCollectorFunction RegisterSlateGlobalPSOCollector(&SlateGlobal
 
 class FSlateMaterialPSOCollector : public IPSOCollector
 {
+	ERHIFeatureLevel::Type FeatureLevel;
+
 public:
 	FSlateMaterialPSOCollector(ERHIFeatureLevel::Type InFeatureLevel) : 
-		IPSOCollector(FPSOCollectorCreateManager::GetIndex(GetFeatureLevelShadingPath(InFeatureLevel), SlateMaterialPSOCollectorName))
+		IPSOCollector(FPSOCollectorCreateManager::GetIndex(GetFeatureLevelShadingPath(InFeatureLevel), SlateMaterialPSOCollectorName)),
+		FeatureLevel(InFeatureLevel)
 	{
 	}
 
@@ -1921,8 +1933,10 @@ void FSlateMaterialPSOCollector::CollectPSOInitializers(
 		return;
 	}
 
+	const EShaderPlatform ShaderPlatform = GetFeatureLevelShaderPlatform(FeatureLevel);
+
 	const bool bRequired = false;
-	for (ESlateShader ShaderType : GetSlateShaderTypesToPrecache())
+	for (ESlateShader ShaderType : GetSlateShaderTypesToPrecache(ShaderPlatform))
 	{
 		int32 UseInstancingCount = ShaderType == ESlateShader::Custom ? 2 : 1;
 		for (int32 UseInstancingIndex = 0; UseInstancingIndex < UseInstancingCount; ++UseInstancingIndex)
@@ -1930,7 +1944,7 @@ void FSlateMaterialPSOCollector::CollectPSOInitializers(
 			bool bUseInstancing = UseInstancingIndex > 0;
 
 			FMaterialShaderTypes ShaderTypesToGet;
-			if (!ChooseMaterialShaderTypes(ShaderType, bUseInstancing, ShaderTypesToGet))
+			if (!ChooseMaterialShaderTypes(ShaderType, ShaderPlatform, bUseInstancing, ShaderTypesToGet))
 			{
 				continue;
 			}
