Move Present Begin/End latency callbacks into FRHIViewport

#Synced-CL 49455230

diff --git a/Engine/Source/Runtime/D3D12RHI/Private/D3D12Viewport.cpp b/Engine/Source/Runtime/D3D12RHI/Private/D3D12Viewport.cpp
--- a/Engine/Source/Runtime/D3D12RHI/Private/D3D12Viewport.cpp
+++ b/Engine/Source/Runtime/D3D12RHI/Private/D3D12Viewport.cpp
@@ -526,6 +526,11 @@ void FD3D12Viewport::PresentChecked(int32 SyncInterval)
 {
 	static constexpr uint32 MaxPresentFailures = 5;
 
+	if (PresentBeginLatencyCallback)
+	{
+		PresentBeginLatencyCallback();
+	}
+
 	// Present the back buffer to the viewport window.
 	// In case presentation failures are transient, don't fault on the first one.
 	HRESULT Result = PresentInternal(SyncInterval);
@@ -550,6 +555,11 @@ void FD3D12Viewport::PresentChecked(int32 SyncInterval)
 		);
 	}
 
+	if (PresentEndLatencyCallback)
+	{
+		PresentEndLatencyCallback();
+	}
+
 #if LOG_PRESENT
 	const FString& ThreadName = FThreadManager::GetThreadName(FPlatformTLS::GetCurrentThreadId());
 	UE_LOG(LogD3D12RHI, Log, TEXT("*** PRESENT: Thread %s: Viewport %#016llx: BackBuffer %#016llx (SyncInterval %u) ***"), ThreadName.GetCharArray().GetData(), this, GetBackBuffer_RHIThread(), SyncInterval);
diff --git a/Engine/Source/Runtime/Engine/Private/UnrealClient.cpp b/Engine/Source/Runtime/Engine/Private/UnrealClient.cpp
--- a/Engine/Source/Runtime/Engine/Private/UnrealClient.cpp
+++ b/Engine/Source/Runtime/Engine/Private/UnrealClient.cpp
@@ -1746,17 +1746,17 @@ void FViewport::EndRenderFrame(FRHICommandListImmediate& RHICmdList, bool bPrese
 {
 	check( IsInRenderingThread() );
 
-	RHICmdList.EnqueueLambda([CurrentFrameCounter = GFrameCounterRenderThread](FRHICommandListImmediate& InRHICmdList)
+	GetViewportRHI()->SetPresentLatencyCallbacks(
+		[CurrentFrameCounter = GFrameCounterRenderThread]()
 	{
 		UEngine::SetPresentLatencyMarkerStart(CurrentFrameCounter);
-	});
-
-	RHICmdList.EndDrawingViewport(GetViewportRHI(), FRHIPresentArgs(/*.FrameCounter*/ GFrameCounterRenderThread, /*.bPresent*/ bPresent, /*.bLockToVsync*/ bLockToVsync));
-
-	RHICmdList.EnqueueLambda([CurrentFrameCounter = GFrameCounterRenderThread](FRHICommandListImmediate& InRHICmdList)
+	},
+	[CurrentFrameCounter = GFrameCounterRenderThread]()
 	{
 		UEngine::SetPresentLatencyMarkerEnd(CurrentFrameCounter);
 	});
+
+	RHICmdList.EndDrawingViewport(GetViewportRHI(), FRHIPresentArgs(/*.FrameCounter*/ GFrameCounterRenderThread, /*.bPresent*/ bPresent, /*.bLockToVsync*/ bLockToVsync));
 }
 
 FRHIGPUMask FViewport::GetGPUMask(FRHICommandListImmediate& RHICmdList) const
diff --git a/Engine/Source/Runtime/RHI/Public/RHIResources.h b/Engine/Source/Runtime/RHI/Public/RHIResources.h
--- a/Engine/Source/Runtime/RHI/Public/RHIResources.h
+++ b/Engine/Source/Runtime/RHI/Public/RHIResources.h
@@ -2544,6 +2544,16 @@ public:
 	virtual void WaitForFrameEventCompletion() { }
 
 	virtual void IssueFrameEvent() { }
+
+	void SetPresentLatencyCallbacks(TFunction<void()>&& InBeginCallback, TFunction<void()>&& InEndCallback)
+	{
+		PresentBeginLatencyCallback = MoveTemp(InBeginCallback);
+		PresentEndLatencyCallback = MoveTemp(InEndCallback);
+	}
+
+protected:
+	TFunction<void()> PresentBeginLatencyCallback;
+	TFunction<void()> PresentEndLatencyCallback;
 };
 
 /** Used to specify a texture metadata plane when creating a view. */
diff --git a/Engine/Source/Runtime/SlateRHIRenderer/Private/SlateRHIRenderer.cpp b/Engine/Source/Runtime/SlateRHIRenderer/Private/SlateRHIRenderer.cpp
--- a/Engine/Source/Runtime/SlateRHIRenderer/Private/SlateRHIRenderer.cpp
+++ b/Engine/Source/Runtime/SlateRHIRenderer/Private/SlateRHIRenderer.cpp
@@ -921,12 +921,7 @@ void FSlateRHIRenderer::PresentWindow_RenderThread(FRHICommandListImmediate& RHI
 	OnBackBufferReadyToPresentDelegate.Broadcast(*DrawPassInputs.Window, DrawPassOutputs.OutputTextureRHI);
 
 	uint32 StartTime = FPlatformTime::Cycles();
-
-	RHICmdList.EnqueueLambda([CurrentFrameCounter = GFrameCounterRenderThread](FRHICommandListImmediate& InRHICmdList)
-	{
-		UEngine::SetPresentLatencyMarkerStart(CurrentFrameCounter);
-	});
-
+	
 	FRHITexture* OptionalSDRBuffer = nullptr;
 
 	if (GRHIGlobals.NeedsExtraTransitions)
@@ -941,12 +936,17 @@ void FSlateRHIRenderer::PresentWindow_RenderThread(FRHICommandListImmediate& RHI
 		}
 	}
 
-	RHICmdList.EndDrawingViewport(DrawPassOutputs.ViewportRHI, FRHIPresentArgs(/*.FrameCounter*/ GFrameCounterRenderThread, /*.bPresent*/ true, /*.bLockToVsync*/ DrawPassInputs.bLockToVsync));
+	DrawPassOutputs.ViewportRHI->SetPresentLatencyCallbacks(
+		[CurrentFrameCounter = GFrameCounterRenderThread]()
+		{
+			UEngine::SetPresentLatencyMarkerStart(CurrentFrameCounter);
+		},
+		[CurrentFrameCounter = GFrameCounterRenderThread]()
+		{
+			UEngine::SetPresentLatencyMarkerEnd(CurrentFrameCounter);
+		});
 
-	RHICmdList.EnqueueLambda([CurrentFrameCounter = GFrameCounterRenderThread](FRHICommandListImmediate& InRHICmdList)
-	{
-		UEngine::SetPresentLatencyMarkerEnd(CurrentFrameCounter);
-	});
+	RHICmdList.EndDrawingViewport(DrawPassOutputs.ViewportRHI, FRHIPresentArgs(/*.FrameCounter*/ GFrameCounterRenderThread, /*.bPresent*/ true, /*.bLockToVsync*/ DrawPassInputs.bLockToVsync));
 
 	uint32 EndTime = FPlatformTime::Cycles();
 
diff --git a/Engine/Source/Runtime/VulkanRHI/Private/VulkanViewport.cpp b/Engine/Source/Runtime/VulkanRHI/Private/VulkanViewport.cpp
--- a/Engine/Source/Runtime/VulkanRHI/Private/VulkanViewport.cpp
+++ b/Engine/Source/Runtime/VulkanRHI/Private/VulkanViewport.cpp
@@ -677,6 +677,11 @@ bool FVulkanViewport::Present(FVulkanCommandListContext& Context, bool bLockToVs
 
 	FPlatformAtomics::AtomicStore(&LockToVsync, bLockToVsync ? 1 : 0);
 
+	if (PresentBeginLatencyCallback)
+	{
+		PresentBeginLatencyCallback();
+	}
+
 	//Transition back buffer to presentable and submit that command
 	if (SupportsStandardSwapchain())
 	{
@@ -778,6 +783,11 @@ bool FVulkanViewport::Present(FVulkanCommandListContext& Context, bool bLockToVs
 		}
 	}
 
+	if (PresentEndLatencyCallback)
+	{
+		PresentEndLatencyCallback();
+	}
+
 	if (FVulkanPlatform::RequiresWaitingForFrameCompletionEvent() && !bHasCustomPresent)
 	{
 		// Wait for the GPU to finish rendering the previous frame before finishing this frame.
diff --git a/Engine/Source/Runtime/Windows/D3D11RHI/Private/D3D11Viewport.cpp b/Engine/Source/Runtime/Windows/D3D11RHI/Private/D3D11Viewport.cpp
--- a/Engine/Source/Runtime/Windows/D3D11RHI/Private/D3D11Viewport.cpp
+++ b/Engine/Source/Runtime/Windows/D3D11RHI/Private/D3D11Viewport.cpp
@@ -386,6 +386,11 @@ bool FD3D11Viewport::PresentChecked(IRHICommandContext& RHICmdContext, int32 Syn
 	}
 #endif
 
+	if (PresentBeginLatencyCallback)
+	{
+		PresentBeginLatencyCallback();
+	}
+
 	if (IsValidRef(CustomPresent))
 	{
 		SCOPE_CYCLE_COUNTER(STAT_D3D11CustomPresentTime);
@@ -446,6 +451,11 @@ bool FD3D11Viewport::PresentChecked(IRHICommandContext& RHICmdContext, int32 Syn
 		}
 	}
 
+	if (PresentEndLatencyCallback)
+	{
+		PresentEndLatencyCallback();
+	}
+
 	FThreadHeartBeat::Get().PresentFrame();
 	if (FAILED(Result))
 	{
