Add slate.Mobile.MaterialShaderMode which can override slate.MaterialShaderMode on mobile platforms.
Default value is -1 which disables any override.
#Synced-CL 49645631

diff --git a/Engine/Source/Runtime/SlateRHIRenderer/Private/SlateRHIRenderingPolicy.cpp b/Engine/Source/Runtime/SlateRHIRenderer/Private/SlateRHIRenderingPolicy.cpp
--- a/Engine/Source/Runtime/SlateRHIRenderer/Private/SlateRHIRenderingPolicy.cpp
+++ b/Engine/Source/Runtime/SlateRHIRenderer/Private/SlateRHIRenderingPolicy.cpp
@@ -60,10 +60,23 @@ static TAutoConsoleVariable<int32> CVarSlateMaterialShaderMode(
 	TEXT(" 2: Only use dynamic branch version of the shaders (lowest shader count)\n"),
 	ECVF_ReadOnly);
 
+static TAutoConsoleVariable<int32> CVarSlateMobileMaterialShaderMode(
+	TEXT("slate.Mobile.MaterialShaderMode"),
+	-1,
+	TEXT("Mobile renderer specific override for slate.MaterialShaderMode.")
+	TEXT(" -1: No mobile override. Uses value from slate.MaterialShaderMode (default).\n")
+	TEXT(" [0-2]: See mode descriptions from slate.MaterialShaderMode.\n"),
+	ECVF_ReadOnly);
+
 static int32 GetSlateMaterialShaderMode(EShaderPlatform ShaderPlatform)
 {
 	static FShaderPlatformCachedIniValue<int32> PerPlatformCVar(TEXT("slate.MaterialShaderMode"));
-	return PerPlatformCVar.Get(ShaderPlatform);
+	const int32 Mode = PerPlatformCVar.Get(ShaderPlatform);
+	static FShaderPlatformCachedIniValue<int32> MobilePerPlatformCVar(TEXT("slate.Mobile.MaterialShaderMode"));
+	const int32 MobileMode = MobilePerPlatformCVar.Get(ShaderPlatform);
+
+	const bool bIsMobile = IsMobilePlatform(ShaderPlatform);
+	return bIsMobile && MobileMode >= 0 ? MobileMode : Mode;
 }
 
 bool IsValidSlateMaterialShaderType(ESlateShader ShaderType, EShaderPlatform ShaderPlatform)
