Adding SlateRHIRenderer support for multi view slate window rendering on macOS

#Synced-CL 49959548

diff --git a/Engine/Source/Runtime/SlateCore/Public/Rendering/SlateRenderer.h b/Engine/Source/Runtime/SlateCore/Public/Rendering/SlateRenderer.h
--- a/Engine/Source/Runtime/SlateCore/Public/Rendering/SlateRenderer.h
+++ b/Engine/Source/Runtime/SlateCore/Public/Rendering/SlateRenderer.h
@@ -499,6 +499,10 @@ public:
 	 * @return	Newly created updatable texture
 	 */
 	virtual FSlateUpdatableTexture* CreateSharedHandleTexture(void *SharedHandle) = 0;
+	
+	virtual void CreateNativeLayer(int32 NewNativeLayer, SWindow& InWindow, void* NativeViewHandle) {};
+	
+	virtual void DeleteNativeLayer(int32 OldNativeLayer, SWindow& InWindow) {};
 
 	/**
 	 * Return an updatable texture to the renderer for release
diff --git a/Engine/Source/Runtime/SlateRHIRenderer/Private/SlateRHIRenderer.cpp b/Engine/Source/Runtime/SlateRHIRenderer/Private/SlateRHIRenderer.cpp
--- a/Engine/Source/Runtime/SlateRHIRenderer/Private/SlateRHIRenderer.cpp
+++ b/Engine/Source/Runtime/SlateRHIRenderer/Private/SlateRHIRenderer.cpp
@@ -61,6 +61,10 @@
 #include "ShaderCompiler.h"
 #endif
 
+#if PLATFORM_MAC
+#include "Misc/ScopeRWLock.h"
+#endif
+
 DECLARE_CYCLE_STAT(TEXT("Total Render Thread time including dependent waits"), STAT_RenderThreadCriticalPath, STATGROUP_Threading);
 
 CSV_DEFINE_CATEGORY(RenderThreadIdle, true);
@@ -394,6 +398,13 @@ void FSlateRHIRenderer::CreateViewport(const TSharedRef<SWindow> Window)
 
 	WindowToViewportInfo.Add(&Window.Get(), ViewInfo);
 	Window->SetIsHDR(ViewInfo->bDisplayFormatIsHDR);
+	
+#if WITH_MULTI_VIEW_SLATE_WINDOW_SUPPORT
+	// create the very base "native layer" used exclusively if we need
+	// no interoperability with native UI elements 
+	// (i.e. only rendering to one native view for Slate)
+	CreateNativeLayer(std::numeric_limits<int32>::min(), *Window, nullptr);
+#endif
 }
 
 void FSlateRHIRenderer::ResizeViewportIfNeeded(FSlateViewportInfo* ViewInfo, FIntPoint ExtentToResizeTo, bool bFullscreen, SWindow* Window)
@@ -683,266 +694,349 @@ public:
 
 IMPLEMENT_GLOBAL_SHADER(FCompositeCS, "/Engine/Private/CompositeUIPixelShader.usf", "CompositeUICS", SF_Compute);
 
-FSlateDrawWindowPassOutputs FSlateRHIRenderer::DrawWindow_RenderThread(FRDGBuilder& GraphBuilder, const FSlateDrawWindowPassInputs& Inputs)
-{
-	LLM_SCOPE(ELLMTag::SceneRender);
+void FSlateRHIRenderer::DrawWindowViewport_RenderThread(ISlateViewport* SlateViewport,
+														FRDGBuilder& GraphBuilder, const FSlateDrawWindowPassInputs& Inputs,
+														FRHITexture** ViewportTextureRHI, FRHITexture** OutputTextureRHI,
+														FSlateViewportInfo& ViewportInfo, FSlateBatchData& BatchData, 
+														FSlateBatchData& BatchDataHDR, TOptional<int32> NativeLayer,
+														TInterval<int32> LayerRange) {
+	// The viewport texture is an optional user-allocated render target. This is rendered to if valid.
+	*ViewportTextureRHI = SlateViewport && SlateViewport->UseSeparateRenderTarget() ? static_cast<FSlateRenderTargetRHI*>(SlateViewport->GetViewportRenderTargetTexture())->GetTypedResource() : nullptr;
 
-	FSlateViewportInfo& ViewportInfo = *Inputs.ViewportInfo;
-	FSlateWindowElementList& WindowElementList = *Inputs.WindowElementList;
+	// The swap chain is the final output. This is rendered to if no viewport render target is provided.
+	FRHITexture* SwapChainTextureRHI = NativeLayer.IsSet() ? RHIGetViewportBackBufferAtLayerIndex(ViewportInfo.ViewportRHI, NativeLayer.GetValue()) : RHIGetViewportBackBuffer(ViewportInfo.ViewportRHI);
 
-	FMaterialRenderProxy::UpdateDeferredCachedUniformExpressions(GraphBuilder.RHICmdList);
-	GetRendererModule().InitializeSystemTextures(GraphBuilder.RHICmdList);
+	// Only render to the intermediate viewport render target if stereo rendering is enabled, which we'll then composite later.
+	const bool bCompositeStereoToSwapChain = *ViewportTextureRHI && GEngine && GEngine->StereoRenderingDevice.IsValid() && SlateViewport && SlateViewport->IsStereoscopic3D();
+	
+	// The output texture is what we ultimately render or composite slate elements into.
+	*OutputTextureRHI = bCompositeStereoToSwapChain ? *ViewportTextureRHI : SwapChainTextureRHI;
+	
+	FString OutputTextureString = NativeLayer.IsSet() ? FString::Printf(TEXT("SlateOutputTexture-%d"), NativeLayer.GetValue()) : FString("SlateOutputTexture");
+	
+	FRDGTexture* OutputTexture = RegisterExternalTexture(GraphBuilder, *OutputTextureRHI, *OutputTextureString);
 
-	FRHITexture* ViewportTextureRHI;
-	FRHITexture* OutputTextureRHI;
+	// The elements texture contains UI elements. It can be the same as the output or allocated separately and composited.
+	FRDGTexture* ElementsTexture = OutputTexture;
+	const FIntPoint OutputExtent = OutputTexture->Desc.Extent;
 
-	{
-		RDG_GPU_MASK_SCOPE(GraphBuilder, FRHIGPUMask::FromIndex(RHIGetViewportNextPresentGPUIndex(ViewportInfo.ViewportRHI)));
-#if WANTS_DRAW_MESH_EVENTS
-		RDG_EVENT_SCOPE_CONDITIONAL_STAT(GraphBuilder,  Inputs.WindowTitle.IsEmpty(), SlateUI, "SlateUI Title = <none>");
-		RDG_EVENT_SCOPE_CONDITIONAL_STAT(GraphBuilder, !Inputs.WindowTitle.IsEmpty(), SlateUI, "SlateUI Title = %s", *Inputs.WindowTitle);
-#else
-		RDG_EVENT_SCOPE_STAT(GraphBuilder, SlateUI, "SlateUI");
-#endif
-		RDG_CSV_STAT_EXCLUSIVE_SCOPE(GraphBuilder, Slate);
-		TRACE_CPUPROFILER_EVENT_SCOPE(Slate::DrawWindow_RenderThread);
+	const bool bCompositeUIWithSceneHDR = ViewportInfo.bDisplayFormatIsHDR && CompositeUIWithSceneHDR();
 
-		ISlateViewport* SlateViewport = Inputs.Window->GetViewport().Get();
-
-		// The viewport texture is an optional user-allocated render target. This is rendered to if valid.
-		ViewportTextureRHI = SlateViewport && SlateViewport->UseSeparateRenderTarget() ? static_cast<FSlateRenderTargetRHI*>(SlateViewport->GetViewportRenderTargetTexture())->GetTypedResource() : nullptr;
+	bool bClearElementsTexture = Inputs.bClear || GSlateWireframe;
 
-		// The swap chain is the final output. This is rendered to if no viewport render target is provided.
-		FRHITexture* SwapChainTextureRHI = RHIGetViewportBackBuffer(ViewportInfo.ViewportRHI);
+#if WITH_SLATE_VISUALIZERS
+	bClearElementsTexture |= CVarShowSlateBatching.GetValueOnRenderThread() != 0 || CVarShowSlateOverdraw.GetValueOnRenderThread() != 0;
+#endif
 
-		// Only render to the intermediate viewport render target if stereo rendering is enabled, which we'll then composite later.
-		const bool bCompositeStereoToSwapChain = ViewportTextureRHI && GEngine && GEngine->StereoRenderingDevice.IsValid() && SlateViewport && SlateViewport->IsStereoscopic3D();
+	if (bCompositeUIWithSceneHDR)
+	{
+		const ETextureCreateFlags WriteMaskFlags = RHISupportsRenderTargetWriteMask(GMaxRHIShaderPlatform) ? ETextureCreateFlags::NoFastClearFinalize | ETextureCreateFlags::DisableDCC : ETextureCreateFlags::None;
+
+		ElementsTexture = GraphBuilder.CreateTexture(
+			FRDGTextureDesc::Create2D(
+				OutputExtent,
+				GetSlateRecommendedColorFormat(),
+				FClearValueBinding::Transparent,
+				ETextureCreateFlags::ShaderResource | ETextureCreateFlags::RenderTargetable | WriteMaskFlags),
+			TEXT("CompositeUIWithSceneHDRTexture"));
+
+		// Force a clear of the UI elements texture to black
+		bClearElementsTexture = true;
+	}
 
-		// The output texture is what we ultimately render or composite slate elements into.
-		OutputTextureRHI = bCompositeStereoToSwapChain ? ViewportTextureRHI : SwapChainTextureRHI;
-		FRDGTexture* OutputTexture = RegisterExternalTexture(GraphBuilder, OutputTextureRHI, TEXT("SlateOutputTexture"));
+	const bool bRequiresVirtualTextureFeedback = BatchData.IsVirtualTextureFeedbackRequired() || BatchDataHDR.IsVirtualTextureFeedbackRequired();
+	if (bRequiresVirtualTextureFeedback)
+	{
+		VirtualTexture::BeginFeedback(GraphBuilder);
+	}
 
-		// The elements texture contains UI elements. It can be the same as the output or allocated separately and composited.
-		FRDGTexture* ElementsTexture = OutputTexture;
-		const FIntPoint OutputExtent = OutputTexture->Desc.Extent;
+	const FSlateElementsBuffers SlateElementsBuffers = BuildSlateElementsBuffers(GraphBuilder, BatchData);
+	const FSlateElementsBuffers SlateElementsBuffersHDR = BuildSlateElementsBuffers(GraphBuilder, BatchDataHDR);
 
-		const bool bCompositeUIWithSceneHDR = ViewportInfo.bDisplayFormatIsHDR && CompositeUIWithSceneHDR();
+	FRDGTexture* SlateStencilTexture = nullptr;
 
-		bool bClearElementsTexture = Inputs.bClear || GSlateWireframe;
+	if (BatchData.IsStencilClippingRequired() || BatchDataHDR.IsStencilClippingRequired())
+	{
+		SlateStencilTexture = GraphBuilder.CreateTexture(FRDGTextureDesc::Create2D(OutputExtent, PF_DepthStencil, FClearValueBinding::DepthZero, GetSlateTransientDepthStencilFlags()), TEXT("SlateDepthStencil"));
+	}
 	
-#if WITH_SLATE_VISUALIZERS
-		bClearElementsTexture |= CVarShowSlateBatching.GetValueOnRenderThread() != 0 || CVarShowSlateOverdraw.GetValueOnRenderThread() != 0;
-#endif
-
-		if (bCompositeUIWithSceneHDR)
-		{
-			const ETextureCreateFlags WriteMaskFlags = RHISupportsRenderTargetWriteMask(GMaxRHIShaderPlatform) ? ETextureCreateFlags::NoFastClearFinalize | ETextureCreateFlags::DisableDCC : ETextureCreateFlags::None;
+	FSlateDrawElementsPassInputs DrawElementsInputs =
+	{
+		  .SceneViewportTexture         = OutputTexture
+		, .ElementsMatrix               = FMatrix44f(ViewportInfo.ProjectionMatrix)
+		, .SceneViewRect                = Inputs.SceneViewRect
+		, .CursorPosition               = Inputs.CursorPosition
+		, .Time                         = Inputs.Time
+		, .HDRDisplayColorGamut         = ViewportInfo.HDRDisplayColorGamut
+		, .UsedSlatePostBuffers         = Inputs.UsedSlatePostBuffers
+		, .ViewportScaleUI              = Inputs.ViewportScaleUI
+		, .bWireframe                   = GSlateWireframe
+		, .bElementsTextureIsHDRDisplay =  ViewportInfo.bDisplayFormatIsHDR
+	};
+	
+	if (ScreenshotState.ViewportToCapture == &ViewportInfo)
+	{
+		// we are in a screenshot, render everything to the base draw layer
+		LayerRange = TInterval<int32>(std::numeric_limits<int32>::min(), std::numeric_limits<int32>::max());
+	}
 
-			ElementsTexture = GraphBuilder.CreateTexture(
-				FRDGTextureDesc::Create2D(
-					OutputExtent,
-					GetSlateRecommendedColorFormat(),
-					FClearValueBinding::Transparent,
-					ETextureCreateFlags::ShaderResource | ETextureCreateFlags::RenderTargetable | WriteMaskFlags),
-				TEXT("CompositeUIWithSceneHDRTexture"));
+	if (bCompositeUIWithSceneHDR)
+	{
+		// Color deficiency correction is performed inside of the CompositeUI pass instead.
+		DrawElementsInputs.bAllowColorDeficiencyCorrection = false;
 
-			// Force a clear of the UI elements texture to black
-			bClearElementsTexture = true;
+		if (!BatchDataHDR.GetRenderBatches().IsEmpty())
+		{
+			DrawElementsInputs.ElementsTexture = OutputTexture;
+			DrawElementsInputs.ElementsLoadAction = ERenderTargetLoadAction::EClear;
+			DrawElementsInputs.ElementsBuffers = SlateElementsBuffersHDR;
+			DrawElementsInputs.StencilTexture = BatchDataHDR.IsStencilClippingRequired() ? SlateStencilTexture : nullptr;
+
+			AddSlateDrawElementsPass(GraphBuilder, *RenderingPolicy, DrawElementsInputs, 
+									 BatchDataHDR.GetRenderBatches(), BatchDataHDR.GetFirstRenderBatchIndex(), 
+									 LayerRange);
 		}
 
-		FSlateBatchData& BatchData = WindowElementList.GetBatchData();
-		FSlateBatchData& BatchDataHDR = WindowElementList.GetBatchDataHDR();
+		DrawElementsInputs.bElementsTextureIsHDRDisplay = false;
+	}
 
-		const bool bRequiresVirtualTextureFeedback = BatchData.IsVirtualTextureFeedbackRequired() || BatchDataHDR.IsVirtualTextureFeedbackRequired();
-		if (bRequiresVirtualTextureFeedback)
-		{
-			VirtualTexture::BeginFeedback(GraphBuilder);
-		}
+	DrawElementsInputs.ElementsTexture = ElementsTexture;
+	DrawElementsInputs.ElementsLoadAction = bClearElementsTexture ? ERenderTargetLoadAction::EClear : ERenderTargetLoadAction::ELoad;
+	DrawElementsInputs.ElementsBuffers = SlateElementsBuffers;
+	DrawElementsInputs.StencilTexture = BatchData.IsStencilClippingRequired() ? SlateStencilTexture : nullptr;
+	
+	AddSlateDrawElementsPass(GraphBuilder, *RenderingPolicy, DrawElementsInputs, 
+							 BatchData.GetRenderBatches(), BatchData.GetFirstRenderBatchIndex(),
+							 LayerRange);
 
-		const FSlateElementsBuffers SlateElementsBuffers = BuildSlateElementsBuffers(GraphBuilder, BatchData);
-		const FSlateElementsBuffers SlateElementsBuffersHDR = BuildSlateElementsBuffers(GraphBuilder, BatchDataHDR);
+	if (bCompositeUIWithSceneHDR)
+	{
+		RDG_EVENT_SCOPE(GraphBuilder, "CompositeUI");
+
+		FRDGTexture* ElementsWriteMaskTexture = nullptr;
 
-		FRDGTexture* SlateStencilTexture = nullptr;
+		FGlobalShaderMap* ShaderMap = GetGlobalShaderMap(GMaxRHIFeatureLevel);
 
-		if (BatchData.IsStencilClippingRequired() || BatchDataHDR.IsStencilClippingRequired())
+		if (RHISupportsRenderTargetWriteMask(GMaxRHIShaderPlatform))
 		{
-			SlateStencilTexture = GraphBuilder.CreateTexture(FRDGTextureDesc::Create2D(OutputExtent, PF_DepthStencil, FClearValueBinding::DepthZero, GetSlateTransientDepthStencilFlags()), TEXT("SlateDepthStencil"));
+			FRenderTargetWriteMask::Decode(GraphBuilder, ShaderMap, MakeArrayView({ ElementsTexture }), ElementsWriteMaskTexture, ETextureCreateFlags::None, TEXT("ElementsWriteMaskTexture"));
 		}
 
-		FSlateDrawElementsPassInputs DrawElementsInputs =
+		static const auto CVarOutputDevice = IConsoleManager::Get().FindTConsoleVariableDataInt(TEXT("r.HDR.Display.OutputDevice"));
+
+		FCompositeShader::FPermutationDomain PermutationVector;
+		PermutationVector.Set<FCompositeShader::FSCRGBEncoding>(ViewportInfo.HDRDisplayOutputFormat == EDisplayOutputFormat::HDR_ACES_1000nit_ScRGB || ViewportInfo.HDRDisplayOutputFormat == EDisplayOutputFormat::HDR_ACES_2000nit_ScRGB);
+		PermutationVector.Set<FCompositeShader::FApplyColorDeficiency>(GSlateColorDeficiencyType != EColorVisionDeficiency::NormalVision && GSlateColorDeficiencySeverity > 0);
+
+		FCompositeShaderCommonParameters CommonParameters;
+		CommonParameters.UIWriteMaskTexture = ElementsWriteMaskTexture;
+		CommonParameters.UITexture = ElementsTexture;
+		CommonParameters.UISampler = TStaticSamplerState<SF_Point>::GetRHI();
+		CommonParameters.UILevel = CVarUILevel.GetValueOnRenderThread();
+		CommonParameters.UILuminance = CVarHDRUILuminance.GetValueOnRenderThread();
+		CommonParameters.UICompositeEOTF = CVarHDRUICompositeEOTF.GetValueOnRenderThread();
+		CommonParameters.UICompositeEOTFGamma = CVarHDRUICompositeEOTFGamma.GetValueOnRenderThread();
+		CommonParameters.UICompositeBlend = CVarHDRUICompositeBlend.GetValueOnRenderThread();
+		const float UICompositeBlendGamma = CVarHDRUICompositeBlendGamma.GetValueOnRenderThread();
+		CommonParameters.UICompositeBlendGamma = UICompositeBlendGamma;
+		CommonParameters.UICompositeBlendInvGamma = 1.f / UICompositeBlendGamma;
+		CommonParameters.ColorVisionDeficiencySeverity = (float)GSlateColorDeficiencySeverity;
+		CommonParameters.ColorVisionDeficiencyType = (float)GSlateColorDeficiencyType;
+		CommonParameters.bSimulateCorrectionWithDeficiency = GSlateShowColorDeficiencyCorrectionWithDeficiency ? 1.0f : 0.0f;
+		CommonParameters.bCorrectDeficiency = GSlateColorDeficiencyCorrection ? 1.0f : 0.0f;
+
+		if (FCompositeCS::IsShaderSupported(GMaxRHIShaderPlatform))
 		{
-			  .SceneViewportTexture         = OutputTexture
-			, .ElementsMatrix               = FMatrix44f(ViewportInfo.ProjectionMatrix)
-			, .SceneViewRect                = Inputs.SceneViewRect
-			, .CursorPosition               = Inputs.CursorPosition
-			, .Time                         = Inputs.Time
-			, .HDRDisplayColorGamut         = ViewportInfo.HDRDisplayColorGamut
-			, .UsedSlatePostBuffers         = Inputs.UsedSlatePostBuffers
-			, .ViewportScaleUI              = Inputs.ViewportScaleUI
-			, .bWireframe                   = GSlateWireframe
-			, .bElementsTextureIsHDRDisplay =  ViewportInfo.bDisplayFormatIsHDR
-		};
-
-		if (bCompositeUIWithSceneHDR)
+			FCompositeCS::FParameters* PassParameters = GraphBuilder.AllocParameters<FCompositeCS::FParameters>();
+			PassParameters->Common = CommonParameters;
+			PassParameters->RWSceneTexture = GraphBuilder.CreateUAV(OutputTexture);
+			PassParameters->SceneTextureDimensions = FVector4f((float)OutputExtent.X, (float)OutputExtent.Y, 1.0f/(float)OutputExtent.X, 1.0f/(float)OutputExtent.Y);
+
+			TShaderMapRef<FCompositeCS> ComputeShader(ShaderMap, PermutationVector);
+
+			FComputeShaderUtils::AddPass(
+				GraphBuilder,
+				RDG_EVENT_NAME("CompositeUI"),
+				ComputeShader,
+				PassParameters,
+				FIntVector(
+					FMath::DivideAndRoundUp<int32>(OutputExtent.X, FCompositeCS::NUM_THREADS_PER_GROUP),
+					FMath::DivideAndRoundUp<int32>(OutputExtent.Y, FCompositeCS::NUM_THREADS_PER_GROUP),
+					1));
+		}
+		else
 		{
-			// Color deficiency correction is performed inside of the CompositeUI pass instead.
-			DrawElementsInputs.bAllowColorDeficiencyCorrection = false;
+			FRDGTexture* ViewportCopyTexture = GraphBuilder.CreateTexture(
+				FRDGTextureDesc::Create2D(
+					OutputExtent,
+					OutputTexture->Desc.Format,
+					FClearValueBinding::Transparent,
+					GetSlateTransientRenderTargetFlags()),
+				TEXT("SlateViewportCopyTexture"));
 
-			if (!BatchDataHDR.GetRenderBatches().IsEmpty())
-			{
-				DrawElementsInputs.ElementsTexture = OutputTexture;
-				DrawElementsInputs.ElementsLoadAction = ERenderTargetLoadAction::EClear;
-				DrawElementsInputs.ElementsBuffers = SlateElementsBuffersHDR;
-				DrawElementsInputs.StencilTexture = BatchDataHDR.IsStencilClippingRequired() ? SlateStencilTexture : nullptr;
+			AddCopyTexturePass(GraphBuilder, OutputTexture, ViewportCopyTexture);
+				
+			const FScreenPassTextureViewport Viewport(OutputTexture);
 
-				AddSlateDrawElementsPass(GraphBuilder, *RenderingPolicy, DrawElementsInputs, BatchDataHDR.GetRenderBatches(), BatchDataHDR.GetFirstRenderBatchIndex());
-			}
+			FCompositePS::FParameters* PassParameters = GraphBuilder.AllocParameters<FCompositePS::FParameters>();
+			PassParameters->Common = CommonParameters;
+			PassParameters->RenderTargets[0] = FRenderTargetBinding(OutputTexture, ERenderTargetLoadAction::ENoAction);
+			PassParameters->SceneTexture = ViewportCopyTexture;
+			PassParameters->SceneSampler = TStaticSamplerState<SF_Point>::GetRHI();
 
-			DrawElementsInputs.bElementsTextureIsHDRDisplay = false;
-		}
-
-		DrawElementsInputs.ElementsTexture = ElementsTexture;
-		DrawElementsInputs.ElementsLoadAction = bClearElementsTexture ? ERenderTargetLoadAction::EClear : ERenderTargetLoadAction::ELoad;
-		DrawElementsInputs.ElementsBuffers = SlateElementsBuffers;
-		DrawElementsInputs.StencilTexture = BatchData.IsStencilClippingRequired() ? SlateStencilTexture : nullptr;
+			TShaderMapRef<FCompositePS> PixelShader(ShaderMap, PermutationVector);
 
-		AddSlateDrawElementsPass(GraphBuilder, *RenderingPolicy, DrawElementsInputs, BatchData.GetRenderBatches(), BatchData.GetFirstRenderBatchIndex());
+			AddDrawScreenPass(GraphBuilder, RDG_EVENT_NAME("CompositeUI"), FScreenPassViewInfo(), Viewport, Viewport, PixelShader, PassParameters);
+		}
+	}
 
-		if (bCompositeUIWithSceneHDR)
-		{
-			RDG_EVENT_SCOPE(GraphBuilder, "CompositeUI");
+	if (bCompositeStereoToSwapChain)
+	{
+		FRDGTexture* SwapChainTexture = RegisterExternalTexture(GraphBuilder, SwapChainTextureRHI, TEXT("StereoSpectatorSwapChainTexture"));
+		GraphBuilder.SetTextureAccessFinal(SwapChainTexture, ERHIAccess::Present);
+		GEngine->StereoRenderingDevice->RenderTexture_RenderThread(GraphBuilder, SwapChainTexture, OutputTexture,
+																   Inputs.WindowElementList->GetWindowSize());
+	}
+	else
+	{
+		GraphBuilder.SetTextureAccessFinal(OutputTexture, ERHIAccess::Present);
+	}
 
-			FRDGTexture* ElementsWriteMaskTexture = nullptr;
+	if (bRequiresVirtualTextureFeedback)
+	{
+		VirtualTexture::EndFeedback(GraphBuilder);
+	}
 
-			FGlobalShaderMap* ShaderMap = GetGlobalShaderMap(GMaxRHIFeatureLevel);
+	OnAddBackBufferReadyToPresentPassDelegate.Broadcast(GraphBuilder, *Inputs.Window, OutputTexture);
 
-			if (RHISupportsRenderTargetWriteMask(GMaxRHIShaderPlatform))
-			{
-				FRenderTargetWriteMask::Decode(GraphBuilder, ShaderMap, MakeArrayView({ ElementsTexture }), ElementsWriteMaskTexture, ETextureCreateFlags::None, TEXT("ElementsWriteMaskTexture"));
-			}
+	if (ScreenshotState.ViewportToCapture == &ViewportInfo)
+	{
+		// Sanity check to make sure the user specified a valid screenshot rect.
+		FIntRect ViewRectClamped;
+		ViewRectClamped.Min = ScreenshotState.ViewRect.Min;
+		ViewRectClamped.Max = ScreenshotState.ViewRect.Max.ComponentMin(OutputExtent);
+		ViewRectClamped.Max = ScreenshotState.ViewRect.Min.ComponentMax(ViewRectClamped.Max);
 
-			static const auto CVarOutputDevice = IConsoleManager::Get().FindTConsoleVariableDataInt(TEXT("r.HDR.Display.OutputDevice"));
-
-			FCompositeShader::FPermutationDomain PermutationVector;
-			PermutationVector.Set<FCompositeShader::FSCRGBEncoding>(ViewportInfo.HDRDisplayOutputFormat == EDisplayOutputFormat::HDR_ACES_1000nit_ScRGB || ViewportInfo.HDRDisplayOutputFormat == EDisplayOutputFormat::HDR_ACES_2000nit_ScRGB);
-			PermutationVector.Set<FCompositeShader::FApplyColorDeficiency>(GSlateColorDeficiencyType != EColorVisionDeficiency::NormalVision && GSlateColorDeficiencySeverity > 0);
-
-			FCompositeShaderCommonParameters CommonParameters;
-			CommonParameters.UIWriteMaskTexture = ElementsWriteMaskTexture;
-			CommonParameters.UITexture = ElementsTexture;
-			CommonParameters.UISampler = TStaticSamplerState<SF_Point>::GetRHI();
-			CommonParameters.UILevel = CVarUILevel.GetValueOnRenderThread();
-			CommonParameters.UILuminance = CVarHDRUILuminance.GetValueOnRenderThread();
-			CommonParameters.UICompositeEOTF = CVarHDRUICompositeEOTF.GetValueOnRenderThread();
-			CommonParameters.UICompositeEOTFGamma = CVarHDRUICompositeEOTFGamma.GetValueOnRenderThread();
-			CommonParameters.UICompositeBlend = CVarHDRUICompositeBlend.GetValueOnRenderThread();
-			const float UICompositeBlendGamma = CVarHDRUICompositeBlendGamma.GetValueOnRenderThread();
-			CommonParameters.UICompositeBlendGamma = UICompositeBlendGamma;
-			CommonParameters.UICompositeBlendInvGamma = 1.f / UICompositeBlendGamma;
-			CommonParameters.ColorVisionDeficiencySeverity = (float)GSlateColorDeficiencySeverity;
-			CommonParameters.ColorVisionDeficiencyType = (float)GSlateColorDeficiencyType;
-			CommonParameters.bSimulateCorrectionWithDeficiency = GSlateShowColorDeficiencyCorrectionWithDeficiency ? 1.0f : 0.0f;
-			CommonParameters.bCorrectDeficiency = GSlateColorDeficiencyCorrection ? 1.0f : 0.0f;
-
-			if (FCompositeCS::IsShaderSupported(GMaxRHIShaderPlatform))
-			{
-				FCompositeCS::FParameters* PassParameters = GraphBuilder.AllocParameters<FCompositeCS::FParameters>();
-				PassParameters->Common = CommonParameters;
-				PassParameters->RWSceneTexture = GraphBuilder.CreateUAV(OutputTexture);
-				PassParameters->SceneTextureDimensions = FVector4f((float)OutputExtent.X, (float)OutputExtent.Y, 1.0f/(float)OutputExtent.X, 1.0f/(float)OutputExtent.Y);
-
-				TShaderMapRef<FCompositeCS> ComputeShader(ShaderMap, PermutationVector);
-
-				FComputeShaderUtils::AddPass(
-					GraphBuilder,
-					RDG_EVENT_NAME("CompositeUI"),
-					ComputeShader,
-					PassParameters,
-					FIntVector(
-						FMath::DivideAndRoundUp<int32>(OutputExtent.X, FCompositeCS::NUM_THREADS_PER_GROUP),
-						FMath::DivideAndRoundUp<int32>(OutputExtent.Y, FCompositeCS::NUM_THREADS_PER_GROUP),
-						1));
-			}
-			else
-			{
-				FRDGTexture* ViewportCopyTexture = GraphBuilder.CreateTexture(
-					FRDGTextureDesc::Create2D(
-						OutputExtent,
-						OutputTexture->Desc.Format,
-						FClearValueBinding::Transparent,
-						GetSlateTransientRenderTargetFlags()),
-					TEXT("SlateViewportCopyTexture"));
-
-				AddCopyTexturePass(GraphBuilder, OutputTexture, ViewportCopyTexture);
-					
-				const FScreenPassTextureViewport Viewport(OutputTexture);
-
-				FCompositePS::FParameters* PassParameters = GraphBuilder.AllocParameters<FCompositePS::FParameters>();
-				PassParameters->Common = CommonParameters;
-				PassParameters->RenderTargets[0] = FRenderTargetBinding(OutputTexture, ERenderTargetLoadAction::ENoAction);
-				PassParameters->SceneTexture = ViewportCopyTexture;
-				PassParameters->SceneSampler = TStaticSamplerState<SF_Point>::GetRHI();
-
-				TShaderMapRef<FCompositePS> PixelShader(ShaderMap, PermutationVector);
-
-				AddDrawScreenPass(GraphBuilder, RDG_EVENT_NAME("CompositeUI"), FScreenPassViewInfo(), Viewport, Viewport, PixelShader, PassParameters);
-			}
-		}
-	
-		if (bCompositeStereoToSwapChain)
+		if (ViewRectClamped != ScreenshotState.ViewRect)
 		{
-			FRDGTexture* SwapChainTexture = RegisterExternalTexture(GraphBuilder, SwapChainTextureRHI, TEXT("StereoSpectatorSwapChainTexture"));
-			GraphBuilder.SetTextureAccessFinal(SwapChainTexture, ERHIAccess::Present);
-			GEngine->StereoRenderingDevice->RenderTexture_RenderThread(GraphBuilder, SwapChainTexture, OutputTexture, WindowElementList.GetWindowSize());
+			UE_LOG(LogSlate, Warning, TEXT("Slate: Screenshot rect max coordinate had to be clamped from [%d, %d] to [%d, %d]"), ScreenshotState.ViewRect.Max.X, ScreenshotState.ViewRect.Max.Y, ViewRectClamped.Max.X, ViewRectClamped.Max.Y);
 		}
-		else
+
+		if (!ViewRectClamped.IsEmpty())
 		{
-			GraphBuilder.SetTextureAccessFinal(OutputTexture, ERHIAccess::Present);
+			AddReadbackTexturePass(GraphBuilder, RDG_EVENT_NAME("ScreenshotReadback"), OutputTexture,
+				[this, OutputTexture, ViewRectClamped, ColorDataHDR = ScreenshotState.ColorDataHDR, ColorData = ScreenshotState.ColorData] (FRHICommandListImmediate& RHICmdList)
+			{
+				if (ColorDataHDR)
+				{
+					RHICmdList.ReadSurfaceData(OutputTexture->GetRHI(), ViewRectClamped, *ColorDataHDR, FReadSurfaceDataFlags(RCM_MinMax));
+				}
+				else
+				{
+					check(ColorData);
+					RHICmdList.ReadSurfaceData(OutputTexture->GetRHI(), ViewRectClamped, *ColorData, FReadSurfaceDataFlags());
+				}
+			});
 		}
+	}
+}
 
-		if (bRequiresVirtualTextureFeedback)
+FSlateDrawWindowPassOutputs FSlateRHIRenderer::DrawWindow_RenderThread(FRDGBuilder& GraphBuilder, const FSlateDrawWindowPassInputs& Inputs)
+{
+	LLM_SCOPE(ELLMTag::SceneRender);
+	
+	FSlateViewportInfo& ViewportInfo = *Inputs.ViewportInfo;
+	
+	FMaterialRenderProxy::UpdateDeferredCachedUniformExpressions(GraphBuilder.RHICmdList);
+	GetRendererModule().InitializeSystemTextures(GraphBuilder.RHICmdList);
+	
+	FSlateDrawWindowPassOutputs Outputs;
+	
+	{
+		RDG_GPU_MASK_SCOPE(GraphBuilder, FRHIGPUMask::FromIndex(RHIGetViewportNextPresentGPUIndex(ViewportInfo.ViewportRHI)));
+#if WANTS_DRAW_MESH_EVENTS
+		RDG_EVENT_SCOPE_CONDITIONAL_STAT(GraphBuilder,  Inputs.WindowTitle.IsEmpty(), SlateUI, "SlateUI Title = <none>");
+		RDG_EVENT_SCOPE_CONDITIONAL_STAT(GraphBuilder, !Inputs.WindowTitle.IsEmpty(), SlateUI, "SlateUI Title = %s", *Inputs.WindowTitle);
+#else
+		RDG_EVENT_SCOPE_STAT(GraphBuilder, SlateUI, "SlateUI");
+#endif
+		RDG_CSV_STAT_EXCLUSIVE_SCOPE(GraphBuilder, Slate);
+		TRACE_CPUPROFILER_EVENT_SCOPE(Slate::DrawWindow_RenderThread);
+		
+		ISlateViewport* SlateViewport = Inputs.Window->GetViewport().Get();
+		
+		FSlateBatchData& BatchData = Inputs.WindowElementList->GetBatchData();
+		FSlateBatchData& BatchDataHDR = Inputs.WindowElementList->GetBatchDataHDR();
+		
+#if WITH_MULTI_VIEW_SLATE_WINDOW_SUPPORT
+		FReadScopeLock ScopeLock(WindowLayerInfosLock);
+		
+		if(!WindowLayerInfos[Inputs.Window].LayersCreatedInViewport)
 		{
-			VirtualTexture::EndFeedback(GraphBuilder);
+			FRHITexture* ViewportTextureRHI;
+			FRHITexture* OutputTextureRHI;
+			
+			DrawWindowViewport_RenderThread(SlateViewport, GraphBuilder, Inputs,
+											&ViewportTextureRHI, &OutputTextureRHI,
+											ViewportInfo, BatchData, BatchDataHDR, 
+											TOptional<int32_t>(), TInterval<int32>(std::numeric_limits<int32>::min(), std::numeric_limits<int32>::max()));
+			
+			return FSlateDrawWindowPassOutputs
+			{
+				.ViewportRHI = ViewportInfo.ViewportRHI,
+				.ViewportTextureRHI = ViewportTextureRHI,
+				.OutputTextureRHI = OutputTextureRHI
+			};
 		}
-
-		OnAddBackBufferReadyToPresentPassDelegate.Broadcast(GraphBuilder, *Inputs.Window, OutputTexture);
-
-		if (ScreenshotState.ViewportToCapture == &ViewportInfo)
+		
+		for(int32 NativeLayerIndex = 0; WindowLayerInfos[Inputs.Window].NativeLayers.Num() > NativeLayerIndex; ++NativeLayerIndex)
 		{
-			// Sanity check to make sure the user specified a valid screenshot rect.
-			FIntRect ViewRectClamped;
-			ViewRectClamped.Min = ScreenshotState.ViewRect.Min;
-			ViewRectClamped.Max = ScreenshotState.ViewRect.Max.ComponentMin(OutputExtent);
-			ViewRectClamped.Max = ScreenshotState.ViewRect.Min.ComponentMax(ViewRectClamped.Max);
-
-			if (ViewRectClamped != ScreenshotState.ViewRect)
+			FRHITexture* ViewportTextureRHI;
+			FRHITexture* OutputTextureRHI;
+			TInterval<int32> LayerRange = WindowLayerInfos[Inputs.Window].LayerRanges[NativeLayerIndex];
+			ViewportInfo.ViewportRHI->SetDefaultNativeLayer(LayerRange.Min);
+			
+			DrawWindowViewport_RenderThread(SlateViewport, GraphBuilder, Inputs,
+											&ViewportTextureRHI, &OutputTextureRHI,
+											ViewportInfo, BatchData, BatchDataHDR, 
+											LayerRange.Min, LayerRange);
+			
+			if(NativeLayerIndex == 0)
 			{
-				UE_LOG(LogSlate, Warning, TEXT("Slate: Screenshot rect max coordinate had to be clamped from [%d, %d] to [%d, %d]"), ScreenshotState.ViewRect.Max.X, ScreenshotState.ViewRect.Max.Y, ViewRectClamped.Max.X, ViewRectClamped.Max.Y);
+				Outputs = FSlateDrawWindowPassOutputs
+				{
+					.ViewportRHI = ViewportInfo.ViewportRHI,
+					.ViewportTextureRHI = ViewportTextureRHI,
+					.OutputTextureRHI = OutputTextureRHI
+				};
 			}
-
-			if (!ViewRectClamped.IsEmpty())
+			else if (ScreenshotState.ViewportToCapture == &ViewportInfo)
 			{
-				AddReadbackTexturePass(GraphBuilder, RDG_EVENT_NAME("ScreenshotReadback"), OutputTexture,
-					[this, OutputTexture, ViewRectClamped, ColorDataHDR = ScreenshotState.ColorDataHDR, ColorData = ScreenshotState.ColorData] (FRHICommandListImmediate& RHICmdList)
-				{
-					if (ColorDataHDR)
-					{
-						RHICmdList.ReadSurfaceData(OutputTexture->GetRHI(), ViewRectClamped, *ColorDataHDR, FReadSurfaceDataFlags(RCM_MinMax));
-					}
-					else
-					{
-						check(ColorData);
-						RHICmdList.ReadSurfaceData(OutputTexture->GetRHI(), ViewRectClamped, *ColorData, FReadSurfaceDataFlags());
-					}
-				});
+				// we are in a screenshot, we just needed to render everything to the base draw layer
+				break;
 			}
 		}
+		
+		ViewportInfo.ViewportRHI->SetDefaultNativeLayer(std::numeric_limits<int32>::min());
 	}
+#else
+		FRHITexture* ViewportTextureRHI;
+		FRHITexture* OutputTextureRHI;
+		
+		DrawWindowViewport_RenderThread(SlateViewport, GraphBuilder, Inputs,
+										&ViewportTextureRHI, &OutputTextureRHI,
+										ViewportInfo, BatchData, BatchDataHDR, 
+										TOptional<int32_t>(), TInterval<int32>(std::numeric_limits<int32>::min(), std::numeric_limits<int32>::max()));
+		
+		Outputs = FSlateDrawWindowPassOutputs
+		{
+			.ViewportRHI = ViewportInfo.ViewportRHI,
+			.ViewportTextureRHI = ViewportTextureRHI,
+			.OutputTextureRHI = OutputTextureRHI
+		};
+	}
+#endif
 
-	FSlateDrawWindowPassOutputs Outputs;
-	Outputs.ViewportRHI = ViewportInfo.ViewportRHI;
-	Outputs.ViewportTextureRHI = ViewportTextureRHI;
-	Outputs.OutputTextureRHI = OutputTextureRHI;
 	return Outputs;
 }
 
@@ -1092,6 +1186,120 @@ struct FSlateDrawWindowsCommand : public TConcurrentLinearObject<FSlateDrawWindo
 	TArray<FRenderThreadUpdateContext, FConcurrentLinearArrayAllocator> DeferredUpdates;
 };
 
+#if WITH_MULTI_VIEW_SLATE_WINDOW_SUPPORT
+void FSlateRHIRenderer::CreateNativeLayer(int32 NewNativeLayer, SWindow& InWindow, void* NativeViewHandle)
+{
+	// we need to write lock ONLY when we are directly mutating WindowLayerInfos
+	// (i.e. NOT during creating native layer in RHI viewport), otherwise we will
+	// get a deadlock with the RHI viewport and the render thread.
+	{
+		FWriteScopeLock ScopeLock(WindowLayerInfosLock);
+		
+		if(!WindowLayerInfos.Contains(&InWindow))
+		{
+			check(NewNativeLayer == std::numeric_limits<int32>::min());
+			
+			WindowLayerInfos.Add(&InWindow, LayerInfo{});
+		}
+		
+		TArray<TInterval<int32>>& LayerRanges = WindowLayerInfos[&InWindow].LayerRanges;
+		TSet<int32_t>& NativeLayers = WindowLayerInfos[&InWindow].NativeLayers;
+		
+		NativeLayers.Add(NewNativeLayer);
+		
+		TArray<int32> NativeLayersSorted = NativeLayers.Array();
+		
+		NativeLayersSorted.Sort([](const int32& A, const int32& B) 
+		{
+			return A > B;
+		});
+		
+		LayerRanges.Empty();
+		
+		for(int32 NativeLayerIndex = 0; NativeLayersSorted.Num() > NativeLayerIndex; ++NativeLayerIndex)
+		{
+			TInterval<int32> LayerRange;
+			
+			if(NativeLayerIndex == 0)
+			{
+				LayerRange.Min = NativeLayersSorted[NativeLayerIndex];
+				LayerRange.Max = std::numeric_limits<int32>::max();
+			}
+			else if(NativeLayerIndex == NativeLayers.Num() - 1)
+			{
+				check(NativeLayersSorted[NativeLayerIndex] == std::numeric_limits<int32>::min());
+				
+				LayerRange.Min = std::numeric_limits<int32>::min();
+				LayerRange.Max = NativeLayersSorted[NativeLayerIndex - 1];
+			}
+			else
+			{
+				LayerRange.Min = NativeLayersSorted[NativeLayerIndex];
+				LayerRange.Max = NativeLayersSorted[NativeLayerIndex - 1];
+			}
+			
+			LayerRanges.Add(LayerRange);
+		}
+		
+		WindowLayerInfos[&InWindow].LayersCreatedInViewport = false;
+	}
+	
+	FSlateViewportInfo* ViewportInfo = WindowToViewportInfo.FindRef(&InWindow);
+	check(ViewportInfo != nullptr);
+	{
+		FReadScopeLock ScopeLock(WindowLayerInfosLock);
+		for(int32_t NativeLayer : WindowLayerInfos[&InWindow].NativeLayers)
+		{
+			// create the native layer within metal if it doesn't exist yet (see MetalViewport)
+			ViewportInfo->ViewportRHI->CreateNativeLayer(NativeLayer, NativeViewHandle);
+		}
+	}		
+	
+	FWriteScopeLock ScopeLock(WindowLayerInfosLock);
+	WindowLayerInfos[&InWindow].LayersCreatedInViewport = true;
+}
+
+void FSlateRHIRenderer::DeleteNativeLayer(int32 OldNativeLayer, SWindow& InWindow)
+{
+	// make sure we aren't deleting the bottom most layer
+	check(OldNativeLayer != std::numeric_limits<int32>::min());
+	
+	FWriteScopeLock ScopeLock(WindowLayerInfosLock);
+	
+	TArray<TInterval<int32>>& LayerRanges = WindowLayerInfos[&InWindow].LayerRanges;
+	TSet<int32_t>& NativeLayers = WindowLayerInfos[&InWindow].NativeLayers;
+	
+	check(NativeLayers.Contains(OldNativeLayer));
+	
+	NativeLayers.Remove(OldNativeLayer);
+	
+	if(FSlateViewportInfo* ViewportInfo = WindowToViewportInfo.FindRef(&InWindow))
+	{
+		ViewportInfo->ViewportRHI->DeleteNativeLayer(OldNativeLayer);
+	}
+	
+	bool FoundLayerRangeToRemove = false;
+	
+	// skip the lowest layer, we can't delete it
+	for(int32 NativeLayerIndex = LayerRanges.Num() - 2; 0 <= NativeLayerIndex; --NativeLayerIndex)
+	{
+		if(OldNativeLayer == LayerRanges[NativeLayerIndex].Min)
+		{
+			// these are reverse sorted (lowest layer, highest layer)
+			check(LayerRanges[NativeLayerIndex + 1].Max == OldNativeLayer);
+			
+			LayerRanges[NativeLayerIndex + 1].Max = LayerRanges[NativeLayerIndex].Max;
+			LayerRanges.RemoveAt(NativeLayerIndex);
+			FoundLayerRangeToRemove = true;
+			
+			break;
+		}
+	}
+	
+	check(FoundLayerRangeToRemove);
+}
+#endif
+
 void FSlateRHIRenderer::DrawWindows_RenderThread(FRHICommandListImmediate& RHICmdList, TConstArrayView<FSlateDrawWindowPassInputs> Windows, TConstArrayView<FRenderThreadUpdateContext> DeferredUpdates)
 {
 	struct FWindowPresentCommand
@@ -1201,6 +1409,7 @@ void FSlateRHIRenderer::DrawWindows_Private(FSlateDrawBuffer& WindowDrawBuffer)
 		for (const TSharedRef<FSlateWindowElementList>& WindowElementListRef : WindowDrawBuffer.GetWindowElementLists())
 		{
 			FSlateWindowElementList* WindowElementList = &(*WindowElementListRef);
+			
 			SWindow* Window = WindowElementList->GetRenderWindow();
 
 			if (!Window)
@@ -1264,6 +1473,12 @@ void FSlateRHIRenderer::DrawWindows_Private(FSlateDrawBuffer& WindowDrawBuffer)
 				Window->ResetViewportScaleUIOverride();
 
 				ElementBatcher->SetCompositeHDRViewports(ViewportInfo->bDisplayFormatIsHDR && CompositeUIWithSceneHDR());
+				
+#if WITH_MULTI_VIEW_SLATE_WINDOW_SUPPORT				
+				check(!WindowLayerInfos[Window].LayerRanges.IsEmpty());
+				check(WindowLayerInfos[Window].LayersCreatedInViewport);
+#endif
+				
 				ElementBatcher->AddElements(*WindowElementList);
 
 				const bool bWindowCanRenderPostProcess       = bAppCanRender && ViewportExtent != FIntPoint::ZeroValue;
@@ -1750,4 +1965,4 @@ void FSlateRHIRenderer::AddWidgetRendererUpdate(const FRenderThreadUpdateContext
 			GraphBuilder.Execute();
 		});
 	}
-}
\ No newline at end of file
+}
diff --git a/Engine/Source/Runtime/SlateRHIRenderer/Private/SlateRHIRenderer.h b/Engine/Source/Runtime/SlateRHIRenderer/Private/SlateRHIRenderer.h
--- a/Engine/Source/Runtime/SlateRHIRenderer/Private/SlateRHIRenderer.h
+++ b/Engine/Source/Runtime/SlateRHIRenderer/Private/SlateRHIRenderer.h
@@ -86,6 +86,11 @@ public:
 	virtual void DestroyCachedFastPathElementData(FSlateCachedElementData* InCachedElementData) override;
 	virtual void EndFrame() const override;
 	virtual void AddWidgetRendererUpdate(const struct FRenderThreadUpdateContext& Context, bool bDeferredRenderTargetUpdate) override;
+	
+#if WITH_MULTI_VIEW_SLATE_WINDOW_SUPPORT
+	virtual void CreateNativeLayer(int32 NewNativeLayer, SWindow& InWindow, void* NativeViewHandle) override;
+	virtual void DeleteNativeLayer(int32 OldNativeLayer, SWindow& InWindow) override;
+#endif
 
 	/**
 	 * Reloads texture resources from disk                   
@@ -109,6 +114,12 @@ public:
 
 private:
 	FSlateDrawWindowPassOutputs DrawWindow_RenderThread(FRDGBuilder& GraphBuilder, const FSlateDrawWindowPassInputs& Inputs);
+	
+	void DrawWindowViewport_RenderThread(ISlateViewport* SlateViewport,
+										 FRDGBuilder& GraphBuilder, const FSlateDrawWindowPassInputs& Inputs,
+										 FRHITexture** ViewportTextureRHI, FRHITexture** OutputTextureRHI,
+										 FSlateViewportInfo& ViewportInfo, FSlateBatchData& BatchData, 
+										 FSlateBatchData& BatchDataHDR, TOptional<int32> NativeLayer, TInterval<int32> LayerRange);
 
 	void PresentWindow_RenderThread(FRHICommandListImmediate& RHICmdList, const FSlateDrawWindowPassInputs& DrawPassInputs, const FSlateDrawWindowPassOutputs& DrawPassOutputs);
 
@@ -192,11 +203,23 @@ private:
 		TArray<FLinearColor>* ColorDataHDR = nullptr;
 
 	} ScreenshotState;
+	
+	struct LayerInfo
+	{
+		TSet<int32> NativeLayers;
+		TArray<TInterval<int32>> LayerRanges;
+		bool LayersCreatedInViewport = false;
+	};
 
 	/** These are state management variables for Scenes on the game thread. A similar copy exists on the RHI Rendering Policy for the rendering thread.*/
 	TArray<FSceneInterface*, TInlineAllocator<4>> ActiveScenes;
 	int32 CurrentSceneIndex = -1;
 	
+#if WITH_MULTI_VIEW_SLATE_WINDOW_SUPPORT
+	TMap<SWindow*, LayerInfo> WindowLayerInfos;
+	FRWLock WindowLayerInfosLock;
+#endif
+	
 	/** Version that increments when it is okay to clean up older cached resources */
 	uint32 ResourceVersion = 0;
 };
diff --git a/Engine/Source/Runtime/SlateRHIRenderer/Private/SlateRHIRenderingPolicy.cpp b/Engine/Source/Runtime/SlateRHIRenderer/Private/SlateRHIRenderingPolicy.cpp
--- a/Engine/Source/Runtime/SlateRHIRenderer/Private/SlateRHIRenderingPolicy.cpp
+++ b/Engine/Source/Runtime/SlateRHIRenderer/Private/SlateRHIRenderingPolicy.cpp
@@ -1467,7 +1467,8 @@ void AddSlateDrawElementsPass(
 	const FSlateRHIRenderingPolicy& RenderingPolicy,
 	const FSlateDrawElementsPassInputs& Inputs,
 	TConstArrayView<FSlateRenderBatch> RenderBatches,
-	int32 FirstBatchIndex)
+	int32 FirstBatchIndex,
+    TInterval<int32> LayerRange)
 {
 	TRACE_CPUPROFILER_EVENT_SCOPE(AddSlateDrawElements);
 	
@@ -1590,12 +1591,10 @@ void AddSlateDrawElementsPass(
 			DrawState.GraphicsPSOInit.DepthStencilState = TStaticDepthStencilState<false, CF_Always>::GetRHI();
 
 			const FSlateRenderBatchOp* RenderBatchOp = RenderBatchHeadOp;
-			const FSlateRenderBatchOp* LastRenderBatchOp = nullptr;
-
+			
 			for (; RenderBatchOp != nullptr; RenderBatchOp = RenderBatchOp->Next)
 			{
 				DrawSlateRenderBatch(RHICmdList, DrawState, *Inputs, *RenderBatchOp);
-				LastRenderBatchOp = RenderBatchOp;
 			}
 		});
 
@@ -1609,8 +1608,14 @@ void AddSlateDrawElementsPass(
 	while (NextRenderBatchIndex != INDEX_NONE)
 	{
 		const FSlateRenderBatch& NextRenderBatch = RenderBatches[NextRenderBatchIndex];
-
+		int32 NextRenderBatchLayer = NextRenderBatch.GetLayer();
+		
 		NextRenderBatchIndex = NextRenderBatch.NextBatchIndex;
+		
+		if(NextRenderBatchLayer < LayerRange.Min || (NextRenderBatchLayer >= LayerRange.Max && LayerRange.Max != std::numeric_limits<int32_t>::max()))
+		{
+			continue;
+		}
 
 		FSlateRenderBatchParameters* NextPassParameters  = LastPassParameters;
 		const FSlateClippingState*   NextClippingState   = NextRenderBatch.ClippingState;
diff --git a/Engine/Source/Runtime/SlateRHIRenderer/Private/SlateRHIRenderingPolicy.h b/Engine/Source/Runtime/SlateRHIRenderer/Private/SlateRHIRenderingPolicy.h
--- a/Engine/Source/Runtime/SlateRHIRenderer/Private/SlateRHIRenderingPolicy.h
+++ b/Engine/Source/Runtime/SlateRHIRenderer/Private/SlateRHIRenderingPolicy.h
@@ -79,4 +79,5 @@ void AddSlateDrawElementsPass(
 	const FSlateRHIRenderingPolicy& RenderingPolicy,
 	const FSlateDrawElementsPassInputs& Inputs,
 	TConstArrayView<FSlateRenderBatch> RenderBatches,
-	int32 FirstBatchIndex);
\ No newline at end of file
+	int32 FirstBatchIndex,
+    TInterval<int32> LayerRange = TInterval<int32>(std::numeric_limits<int32>::min(), std::numeric_limits<int32>::max()));
