Pipe trough the Camera Cut information from Gameplay to RHIs Viewport so we can decide if we want skip a Present setting r.SkipPresentOnCameraCut to 1

#Synced-CL 49985560

diff --git a/Engine/Source/Runtime/Core/Private/HAL/ConsoleManager.cpp b/Engine/Source/Runtime/Core/Private/HAL/ConsoleManager.cpp
--- a/Engine/Source/Runtime/Core/Private/HAL/ConsoleManager.cpp
+++ b/Engine/Source/Runtime/Core/Private/HAL/ConsoleManager.cpp
@@ -4208,6 +4208,13 @@ static TAutoConsoleVariable<int32> CVarSetVSyncEditorEnabled(
 	ECVF_RenderThreadSafe);
 #endif
 
+static TAutoConsoleVariable<int32> CVarSkipPresentOnCameraCut(
+	TEXT("r.SkipPresentOnCameraCut"),
+	0,
+	TEXT("0: Do not skip the Present call on a Camera Cut. (default)\n"
+		"1: Skip the Present call on a Camera Cut."),
+	ECVF_Scalability | ECVF_RenderThreadSafe);
+
 static TAutoConsoleVariable<int32> CVarFinishCurrentFrame(
 	TEXT("r.FinishCurrentFrame"),
 	0,
diff --git a/Engine/Source/Runtime/D3D12RHI/Private/D3D12Viewport.cpp b/Engine/Source/Runtime/D3D12RHI/Private/D3D12Viewport.cpp
--- a/Engine/Source/Runtime/D3D12RHI/Private/D3D12Viewport.cpp
+++ b/Engine/Source/Runtime/D3D12RHI/Private/D3D12Viewport.cpp
@@ -461,10 +461,10 @@ FRHICustomPresent* FD3D12Viewport::GetCustomPresent() const
 }
 
 /** Update the expected next present GPU back buffer index from RenderThread point of view */
-void FD3D12Viewport::AdvanceExpectedBackBufferIndex_RenderThread()
+void FD3D12Viewport::AdvanceExpectedBackBufferIndex_RenderThread(bool bPresent)
 {
 	bool bNeedsNativePresent = IsValidRef(CustomPresent) ? 
-		CustomPresent->NeedsNativePresent() || CustomPresent->NeedsAdvanceBackbuffer() : true;
+		CustomPresent->NeedsNativePresent() || CustomPresent->NeedsAdvanceBackbuffer() : bPresent;
 
 	if (bNeedsNativePresent && IsPresentAllowed())
 	{
@@ -708,8 +708,11 @@ bool FD3D12Viewport::EndRendering(FD3D12CommandContextBase& ContextBase, const F
 	
 	if (!IsPresentAllowed() || !InPresentArgs.bPresent )
 	{
+		bLastFramePresented = false;
 		return false;
 	}
+	
+	bLastFramePresented = true;
 
 	check(ContextBase.GetParentAdapter() == GetParentAdapter());
 
@@ -921,7 +924,7 @@ void FD3D12DynamicRHI::RHIEndDrawingViewport(FRHICommandListImmediate& RHICmdLis
 
 			const bool bNativelyPresented = Viewport->EndRendering(Context, PresentArgs);
 
-			if (bNativelyPresented || GUseSubmissionThreadPresent)
+			if ((bNativelyPresented || GUseSubmissionThreadPresent) || (PresentArgs.bPresent == false))
 			{
 				FRenderThreadIdleScope Scope(ERenderThreadIdleTypes::WaitingForAllOtherSleep);
 				static const auto CFinishFrameVar = IConsoleManager::Get().FindTConsoleVariableDataInt(TEXT("r.FinishCurrentFrame"));
@@ -972,7 +975,7 @@ void FD3D12DynamicRHI::RHIEndDrawingViewport(FRHICommandListImmediate& RHICmdLis
 	// But advance the expected present GPU index so the next call to RHIGetViewportNextPresentGPUIndex returns the expected GPU index for the next present.
 	// Warning: when present fails or is not called on the RHIThread then this might not be in sync but RHI thread will fix up the correct state
 	//          Present doesn't happen so shouldn't matter that the index was wrong then
-	ResourceCast(ViewportRHI)->AdvanceExpectedBackBufferIndex_RenderThread();
+	ResourceCast(ViewportRHI)->AdvanceExpectedBackBufferIndex_RenderThread(PresentArgs.bPresent);
 }
 
 uint32 FD3D12DynamicRHI::RHIGetViewportNextPresentGPUIndex(FRHIViewport* ViewportRHI)
diff --git a/Engine/Source/Runtime/D3D12RHI/Private/D3D12Viewport.h b/Engine/Source/Runtime/D3D12RHI/Private/D3D12Viewport.h
--- a/Engine/Source/Runtime/D3D12RHI/Private/D3D12Viewport.h
+++ b/Engine/Source/Runtime/D3D12RHI/Private/D3D12Viewport.h
@@ -150,11 +150,13 @@ public:
 	bool CurrentOutputSupportsHDR() const;
 
 	/** Advance and get the next present GPU index */
-	void AdvanceExpectedBackBufferIndex_RenderThread();
+	void AdvanceExpectedBackBufferIndex_RenderThread(bool bPresent);
 
 	void OnResumeRendering();
 	void OnSuspendRendering();
 
+	bool PresentedLastFrame(){ return bLastFramePresented; }
+
 private:
 	bool IsPresentAllowed();
 	bool CheckSwapChain();
@@ -224,6 +226,7 @@ private:
 	bool bIsValid = true;
 	bool bAllowTearing = true;
 	bool bNeedSwapChain = false;
+	bool bLastFramePresented = true;
 
 	uint32 CheckedPresentFailureCounter = 0;
 
diff --git a/Engine/Source/Runtime/Engine/Private/Slate/SceneViewport.cpp b/Engine/Source/Runtime/Engine/Private/Slate/SceneViewport.cpp
--- a/Engine/Source/Runtime/Engine/Private/Slate/SceneViewport.cpp
+++ b/Engine/Source/Runtime/Engine/Private/Slate/SceneViewport.cpp
@@ -48,6 +48,7 @@ FSceneViewport::FSceneViewport( FViewportClient* InViewportClient, TSharedPtr<SV
 	, bIsCursorVisible( true )
 	, bShouldCaptureMouseOnActivate( true )
 	, bRequiresVsync( false )
+	, bCameraCut( false )
 	, bUseSeparateRenderTarget( InViewportWidget.IsValid() ? !InViewportWidget->ShouldRenderDirectly() : true )
 	, bForceSeparateRenderTarget( false )
 	, bIsResizing( false )
@@ -1424,6 +1425,13 @@ void FSceneViewport::OnViewportDeactivated(const FWindowActivateEvent& InActivat
 	}
 }
 
+
+bool FSceneViewport::IsCameraCut() const
+{
+	return bCameraCut;
+}
+
+
 FSlateShaderResource* FSceneViewport::GetViewportRenderTargetTexture() const
 {
 	if (IsInRenderingThread())
diff --git a/Engine/Source/Runtime/Engine/Private/UnrealClient.cpp b/Engine/Source/Runtime/Engine/Private/UnrealClient.cpp
--- a/Engine/Source/Runtime/Engine/Private/UnrealClient.cpp
+++ b/Engine/Source/Runtime/Engine/Private/UnrealClient.cpp
@@ -1832,10 +1832,12 @@ void FViewport::Draw( bool bShouldPresent /*= true */)
 			{
 				static const auto CVar = IConsoleManager::Get().FindTConsoleVariableDataInt(TEXT("r.VSync"));
 				bool bLockToVsync = CVar->GetValueOnGameThread() != 0;
+				bool bCameraCut = false;
 				ULocalPlayer* Player = (GEngine && World) ? GEngine->GetFirstGamePlayer(World) : NULL;
-				if ( Player )
+				if (Player && Player->PlayerController)
 				{
-					bLockToVsync |= (Player && Player->PlayerController && Player->PlayerController->bCinematicMode);
+					bLockToVsync |= (Player->PlayerController->bCinematicMode);
+					bCameraCut |= (Player->PlayerController->PlayerCameraManager && Player->PlayerController->PlayerCameraManager->bGameCameraCutThisFrame);
 				}
 	 			EnqueueBeginRenderFrame(bShouldPresent);
 
@@ -1885,6 +1887,7 @@ void FViewport::Draw( bool bShouldPresent /*= true */)
 	
 				// Slate doesn't present immediately. Tag the viewport as requiring vsync so that it happens.
 				SetRequiresVsync(bLockToVsync);
+				SetCameraCut(bCameraCut);
 				EnqueueEndRenderFrame(bLockToVsync, bShouldPresent);
 
 				GInputLatencyTimer.GameThreadTrigger = false;
diff --git a/Engine/Source/Runtime/Engine/Public/Slate/SceneViewport.h b/Engine/Source/Runtime/Engine/Public/Slate/SceneViewport.h
--- a/Engine/Source/Runtime/Engine/Public/Slate/SceneViewport.h
+++ b/Engine/Source/Runtime/Engine/Public/Slate/SceneViewport.h
@@ -172,6 +172,11 @@ public:
 	 */
 	virtual void SetRequiresVsync(bool bShouldVsync) override { bRequiresVsync = bShouldVsync; }
 
+	/**
+	* Indicate that the viewport had a camera cut.
+	*/
+	virtual void SetCameraCut(bool bInCameraCut) override { bCameraCut = bInCameraCut; }
+
 	/**
 	 * Returns true if the viewport should be vsynced.
 	 */
@@ -218,6 +223,7 @@ public:
 	ENGINE_API virtual void UpdateViewportRHI(bool bDestroyed, uint32 NewSizeX, uint32 NewSizeY, EWindowMode::Type NewWindowMode, EPixelFormat PreferredPixelFormat) override;
 
 	/** ISlateViewport interface */
+	ENGINE_API virtual bool IsCameraCut() const override;
 	ENGINE_API virtual FSlateShaderResource* GetViewportRenderTargetTexture() const override;
 	ENGINE_API virtual void OnDrawViewport( const FGeometry& AllottedGeometry, const FSlateRect& MyCullingRect, FSlateWindowElementList& OutDrawElements, int32 LayerId, const FWidgetStyle& InWidgetStyle, bool bParentEnabled ) override;
 	ENGINE_API virtual FCursorReply OnCursorQuery( const FGeometry& MyGeometry, const FPointerEvent& CursorEvent ) override;
@@ -429,6 +435,8 @@ private:
 	bool bShouldCaptureMouseOnActivate;
 	/** true if this viewport requires vsync. */
 	bool bRequiresVsync;
+	/** true if this viewport is in camera cut. */
+	bool bCameraCut;
 	/** true if this viewport renders to a separate render target.  false to render directly to the windows back buffer */
 	bool bUseSeparateRenderTarget;
 	/** True if we should force use of a separate render target because the HMD needs it. */
diff --git a/Engine/Source/Runtime/Engine/Public/UnrealClient.h b/Engine/Source/Runtime/Engine/Public/UnrealClient.h
--- a/Engine/Source/Runtime/Engine/Public/UnrealClient.h
+++ b/Engine/Source/Runtime/Engine/Public/UnrealClient.h
@@ -505,6 +505,12 @@ public:
 	 */
 	virtual void SetRequiresVsync(bool bShouldVsync) {}
 
+	/**
+	 * Indicate that this frame is a camera cut.
+	*/
+	virtual void SetCameraCut(bool bInCameraCut) {}
+
+
 	/**
 	 * Sets PreCapture coordinates from the current position of the slate cursor.
 	 */
diff --git a/Engine/Source/Runtime/SlateCore/Public/Rendering/RenderingCommon.h b/Engine/Source/Runtime/SlateCore/Public/Rendering/RenderingCommon.h
--- a/Engine/Source/Runtime/SlateCore/Public/Rendering/RenderingCommon.h
+++ b/Engine/Source/Runtime/SlateCore/Public/Rendering/RenderingCommon.h
@@ -513,6 +513,15 @@ public:
 	{
 		return false;
 	}
+
+	/**
+	* Does this frame identify as a Camera Cut ?
+	*/
+	virtual bool IsCameraCut() const
+	{
+		return false;
+	}
+
 	
 	/**
 	 * Is the viewport expected to render to GetViewportRenderTargetTexture? Or to the swapchain?
diff --git a/Engine/Source/Runtime/SlateRHIRenderer/Private/SlateRHIRenderer.cpp b/Engine/Source/Runtime/SlateRHIRenderer/Private/SlateRHIRenderer.cpp
--- a/Engine/Source/Runtime/SlateRHIRenderer/Private/SlateRHIRenderer.cpp
+++ b/Engine/Source/Runtime/SlateRHIRenderer/Private/SlateRHIRenderer.cpp
@@ -186,6 +186,21 @@ static bool IsVSyncRequired(const FSlateElementBatcher& ElementBatcher)
 	return bLockToVsync;
 }
 
+static bool ShouldPresent(const bool IsCameraCut)
+{
+	bool bPresent = true;
+
+	{
+		static IConsoleVariable* CVar = IConsoleManager::Get().FindConsoleVariable(TEXT("r.SkipPresentOnCameraCut"));
+		if ((CVar->GetInt() != 0) && IsCameraCut)
+		{
+			bPresent = false;
+		}
+	}
+
+	return bPresent;
+}
+
 struct FSlateViewportInfo : public FRenderResource
 {
 	FViewportRHIRef ViewportRHI;
@@ -223,6 +238,7 @@ struct FSlateDrawWindowPassInputs
 	FGameTime Time;
 	bool bLockToVsync = false;
 	bool bClear = false;
+	bool bPresent = true;
 };
 
 struct FSlateDrawWindowPassOutputs
@@ -1070,7 +1086,7 @@ void FSlateRHIRenderer::PresentWindow_RenderThread(FRHICommandListImmediate& RHI
 			UEngine::SetPresentLatencyMarkerEnd(CurrentFrameCounter);
 		});
 
-	RHICmdList.EndDrawingViewport(DrawPassOutputs.ViewportRHI, FRHIPresentArgs(/*.FrameCounter*/ GFrameCounterRenderThread, /*.bPresent*/ true, /*.bLockToVsync*/ DrawPassInputs.bLockToVsync));
+	RHICmdList.EndDrawingViewport(DrawPassOutputs.ViewportRHI, FRHIPresentArgs(/*.FrameCounter*/ GFrameCounterRenderThread, /*.bPresent*/ DrawPassInputs.bPresent, /*.bLockToVsync*/ DrawPassInputs.bLockToVsync));
 
 	uint32 EndTime = FPlatformTime::Cycles();
 
@@ -1398,6 +1414,7 @@ void FSlateRHIRenderer::DrawWindows_Private(FSlateDrawBuffer& WindowDrawBuffer)
 		ESlatePostRT PostProcessCustumDrawBits;
 		FIntPoint CursorPosition;
 		bool bLockToVsync;
+		bool bPresent;
 	};
 
 	TArray<FWindowToRender, FConcurrentLinearArrayAllocator> WindowsToRender;
@@ -1433,6 +1450,7 @@ void FSlateRHIRenderer::DrawWindows_Private(FSlateDrawBuffer& WindowDrawBuffer)
 			FIntPoint ViewportExtent = FIntPoint::ZeroValue;
 			FIntRect ViewportRect;
 			float ViewportScaleUI = Window->GetViewportScaleUIOverride();
+			bool bCameraCut = false;
 
 			if (ISlateViewport* Viewport = Window->GetViewport().Get())
 			{
@@ -1460,6 +1478,8 @@ void FSlateRHIRenderer::DrawWindows_Private(FSlateDrawBuffer& WindowDrawBuffer)
 				{
 					ViewportScaleUI = GetDefault<UUserInterfaceSettings>()->GetDPIScaleBasedOnSize(ViewportExtent);
 				}
+
+				bCameraCut = Viewport->IsCameraCut();
 			}
 
 			if (FSlateViewportInfo* ViewportInfo = WindowToViewportInfo.FindRef(Window))
@@ -1485,6 +1505,7 @@ void FSlateRHIRenderer::DrawWindows_Private(FSlateDrawBuffer& WindowDrawBuffer)
 				const ESlatePostRT PostProcessUsedBits       = bWindowCanRenderPostProcess ? ElementBatcher->GetUsedSlatePostBuffers() : ESlatePostRT::None;
 				const ESlatePostRT PostProcessCustomDrawBits = bWindowCanRenderPostProcess ? ElementBatcher->GetResourceUpdatingPostBuffers() : ESlatePostRT::None;
 				const bool bLockToVsync                      = IsVSyncRequired(*ElementBatcher);
+				const bool bPresent                          = ShouldPresent(bCameraCut);
 
 				ElementBatcher->ResetBatches();
 				ElementBatcher->SetCompositeHDRViewports(false);
@@ -1502,6 +1523,7 @@ void FSlateRHIRenderer::DrawWindows_Private(FSlateDrawBuffer& WindowDrawBuffer)
 					, .PostProcessCustumDrawBits = PostProcessCustomDrawBits
 					, .CursorPosition            = FIntPoint(ViewportCursorPosition.X, ViewportCursorPosition.Y)
 					, .bLockToVsync              = bLockToVsync
+					, .bPresent					 = bPresent
 				});
 
 				PostProcessAnyUsedBits |= PostProcessUsedBits;
@@ -1585,6 +1607,7 @@ void FSlateRHIRenderer::DrawWindows_Private(FSlateDrawBuffer& WindowDrawBuffer)
 			#if ALPHA_BLENDED_WINDOWS
 				, .bClear                    = WindowToRender.Window->GetTransparencySupport() == EWindowTransparency::PerPixel
 			#endif
+				, .bPresent                  = WindowToRender.bPresent
 			});
 		}
 	}
diff --git a/Engine/Source/Runtime/VulkanRHI/Private/VulkanRHI.cpp b/Engine/Source/Runtime/VulkanRHI/Private/VulkanRHI.cpp
--- a/Engine/Source/Runtime/VulkanRHI/Private/VulkanRHI.cpp
+++ b/Engine/Source/Runtime/VulkanRHI/Private/VulkanRHI.cpp
@@ -1275,7 +1275,7 @@ void FVulkanDynamicRHI::RHIEndDrawingViewport(FRHICommandListImmediate& RHICmdLi
 			FVulkanCommandBuffer& CommandBuffer = Context.GetCommandBuffer();
 			check(!CommandBuffer.HasEnded() && !CommandBuffer.IsInsideRenderPass());
 
-			const bool bNativePresent = Viewport->Present(Context, PresentArgs.bLockToVsync);
+			const bool bNativePresent = Viewport->Present(Context, PresentArgs);
 			if (bNativePresent)
 			{
 				//#todo-rco: Check for r.FinishCurrentFrame
@@ -1284,7 +1284,7 @@ void FVulkanDynamicRHI::RHIEndDrawingViewport(FRHICommandListImmediate& RHICmdLi
 	}
 
 	// Switch to the next viewport back buffer
-	ResourceCast(ViewportRHI)->AdvanceBackBufferFrame(RHICmdList);
+	ResourceCast(ViewportRHI)->AdvanceBackBufferFrame(RHICmdList, PresentArgs.bPresent);
 }
 
 #if WITH_RHI_BREADCRUMBS
diff --git a/Engine/Source/Runtime/VulkanRHI/Private/VulkanViewport.cpp b/Engine/Source/Runtime/VulkanRHI/Private/VulkanViewport.cpp
--- a/Engine/Source/Runtime/VulkanRHI/Private/VulkanViewport.cpp
+++ b/Engine/Source/Runtime/VulkanRHI/Private/VulkanViewport.cpp
@@ -265,11 +265,11 @@ FTextureRHIRef FVulkanViewport::GetBackBuffer(FRHICommandListImmediate& RHICmdLi
 	return RenderingBackBuffer.GetReference();
 }
 
-void FVulkanViewport::AdvanceBackBufferFrame(FRHICommandListImmediate& RHICmdList)
+void FVulkanViewport::AdvanceBackBufferFrame(FRHICommandListImmediate& RHICmdList, bool bPresent)
 {
 	check(IsInRenderingThread());
 
-	if (SupportsStandardSwapchain() && GVulkanDelayAcquireImage != EDelayAcquireImageType::DelayAcquire)
+	if (SupportsStandardSwapchain() && GVulkanDelayAcquireImage != EDelayAcquireImageType::DelayAcquire && bPresent)
 	{
 		check(RHIBackBuffer);
 		
@@ -671,11 +671,16 @@ inline static void CopyImageToBackBuffer(FVulkanCommandListContext& Context, FVu
 	}
 }
 
-bool FVulkanViewport::Present(FVulkanCommandListContext& Context, bool bLockToVsync)
+bool FVulkanViewport::Present(FVulkanCommandListContext& Context, const FRHIPresentArgs& InPresentArgs)
 {
 	check(Context.IsImmediate());
 
-	FPlatformAtomics::AtomicStore(&LockToVsync, bLockToVsync ? 1 : 0);
+	if (!InPresentArgs.bPresent)
+	{
+		return false;
+	}
+
+	FPlatformAtomics::AtomicStore(&LockToVsync, InPresentArgs.bLockToVsync ? 1 : 0);
 
 	if (PresentBeginLatencyCallback)
 	{
diff --git a/Engine/Source/Runtime/VulkanRHI/Private/VulkanViewport.h b/Engine/Source/Runtime/VulkanRHI/Private/VulkanViewport.h
--- a/Engine/Source/Runtime/VulkanRHI/Private/VulkanViewport.h
+++ b/Engine/Source/Runtime/VulkanRHI/Private/VulkanViewport.h
@@ -49,7 +49,7 @@ public:
 	~FVulkanViewport();
 
 	FTextureRHIRef GetBackBuffer(FRHICommandListImmediate& RHICmdList);
-	void AdvanceBackBufferFrame(FRHICommandListImmediate& RHICmdList);
+	void AdvanceBackBufferFrame(FRHICommandListImmediate& RHICmdList, bool bPresent);
 
 	virtual void WaitForFrameEventCompletion() override;
 
@@ -72,7 +72,7 @@ public:
 
 	virtual void Tick(float DeltaTime) override final;
 	
-	bool Present(FVulkanCommandListContext& Context, bool bLockToVsync);
+	bool Present(FVulkanCommandListContext& Context, const FRHIPresentArgs& InPresentArgs);
 
 	inline uint32 GetPresentCount() const
 	{
