Support system HDR paper white values

This value is used to adjust overall 3D scene brightness, and can override r.HDR.UI.Luminance if r.HDR.UI.UseOSPaperWhite is enabled.

Also fix fallback values in SlatePostProcessor.cpp to align with CVar defaults.

#Synced-CL 50431656

diff --git a/Engine/Shaders/Private/PostProcessCombineLUTs.usf b/Engine/Shaders/Private/PostProcessCombineLUTs.usf
--- a/Engine/Shaders/Private/PostProcessCombineLUTs.usf
+++ b/Engine/Shaders/Private/PostProcessCombineLUTs.usf
@@ -47,7 +47,9 @@ float3x3 GetLimitingXYZToRgb()
 	return float3x3(LimitingXYZToRgb0, LimitingXYZToRgb1, LimitingXYZToRgb2);
 }
 
-uint bUseMobileTonemapper;
+float HDRPaperWhiteRatio;
+float InvHDRPaperWhiteRatio;
+
 uint bIsTemperatureWhiteBalance;
 
 float LUTSize;
@@ -459,7 +461,8 @@ float4 CombineLUTsCommon(float3 LUTEncodedColor)
 	// Transform with PQ/2084 encoding, user specified gamut
 	else BRANCH if( GetOutputDevice() == TONEMAPPER_OUTPUT_ACES1000nitST2084 || GetOutputDevice() == TONEMAPPER_OUTPUT_ACES2000nitST2084)
 	{
-		float3 OutputGamutColor = ApplyToneCurve(ColorAP1, OutputMaxLuminance, AP1_2_Output);
+		float3 OutputGamutColor = ApplyToneCurve(ColorAP1, OutputMaxLuminance * InvHDRPaperWhiteRatio, AP1_2_Output);
+		OutputGamutColor *= HDRPaperWhiteRatio;
 
 		// Apply conversion to ST-2084 (Dolby PQ)
 		OutDeviceColor = LinearToST2084(OutputGamutColor * LinearToNitsScale);
@@ -468,7 +471,8 @@ float4 CombineLUTsCommon(float3 LUTEncodedColor)
 	// Transform to linear ScRGB
 	else BRANCH if( GetOutputDevice() == TONEMAPPER_OUTPUT_ACES1000nitScRGB || GetOutputDevice() == TONEMAPPER_OUTPUT_ACES2000nitScRGB)
 	{
-		float3 OutputGamutColor = ApplyToneCurve(ColorAP1, OutputMaxLuminance, AP1_2_Output);
+		float3 OutputGamutColor = ApplyToneCurve(ColorAP1, OutputMaxLuminance * InvHDRPaperWhiteRatio, AP1_2_Output);
+		OutputGamutColor *= HDRPaperWhiteRatio;
 		OutDeviceColor = (LinearToNitsScale / UE_SCRGB_WHITE_NITS) * OutputGamutColor;
 	}
 
diff --git a/Engine/Source/Runtime/Core/Public/HAL/Platform.h b/Engine/Source/Runtime/Core/Public/HAL/Platform.h
--- a/Engine/Source/Runtime/Core/Public/HAL/Platform.h
+++ b/Engine/Source/Runtime/Core/Public/HAL/Platform.h
@@ -596,6 +596,10 @@
 	#define PLATFORM_USES_FIXED_HDR_SETTING 0
 #endif
 
+#ifndef PLATFORM_PROVIDES_HDR_PAPER_WHITE
+	#define PLATFORM_PROVIDES_HDR_PAPER_WHITE 0
+#endif
+
 #ifndef PLATFORM_MANAGES_HDR_SETTING
 	#define PLATFORM_MANAGES_HDR_SETTING 0
 #endif
diff --git a/Engine/Source/Runtime/Engine/Private/Slate/SceneViewport.cpp b/Engine/Source/Runtime/Engine/Private/Slate/SceneViewport.cpp
--- a/Engine/Source/Runtime/Engine/Private/Slate/SceneViewport.cpp
+++ b/Engine/Source/Runtime/Engine/Private/Slate/SceneViewport.cpp
@@ -67,6 +67,7 @@ FSceneViewport::FSceneViewport( FViewportClient* InViewportClient, TSharedPtr<SV
 	, MinimumLuminanceInNits(0.f)
 	, MaximumLuminanceInNits(100.f)
 	, MaximumFullFrameLuminanceInNits(100.f)
+	, HDRPaperWhiteInNits(203.f)
 	, LimitingColorSpace(UE::Color::EColorSpace::sRGB)
 {
 	bIsSlateViewport = true;
@@ -424,6 +425,7 @@ static void ComputeSceneViewportHDRMetaData(FHDRMetaData& OutHDRMetaData, const
 		OutHDRMetaData.MinimumLuminanceInNits = OutHDRMetaData.bHDRSupported ? .1f : 0.f;
 		OutHDRMetaData.MaximumLuminanceInNits = OutHDRMetaData.bHDRSupported ? 1000.f : 100.f;
 		OutHDRMetaData.MaximumFullFrameLuminanceInNits = OutHDRMetaData.bHDRSupported ? 600.f : 100.f;
+		OutHDRMetaData.HDRPaperWhiteInNits = 203.f;
 		OutHDRMetaData.LimitingColorSpace = OutHDRMetaData.bHDRSupported ? UE::Color::FColorSpace::GetP3D65() : UE::Color::FColorSpace::GetSRGB();
 		return;
 	}
@@ -468,6 +470,7 @@ void FSceneViewport::OnDrawViewport( const FGeometry& AllottedGeometry, const FS
 			bHDRStale |= MinimumLuminanceInNits != OutHDRMetaData.MinimumLuminanceInNits;
 			bHDRStale |= MaximumLuminanceInNits != OutHDRMetaData.MaximumLuminanceInNits;
 			bHDRStale |= MaximumFullFrameLuminanceInNits != OutHDRMetaData.MaximumFullFrameLuminanceInNits;
+			bHDRStale |= HDRPaperWhiteInNits != OutHDRMetaData.HDRPaperWhiteInNits;
 			bHDRStale |= LimitingColorSpace != OutHDRMetaData.LimitingColorSpace;
 			bHDRStale |= bHDRViewport != OutHDRMetaData.bHDRSupported;
 		}
@@ -1392,6 +1395,11 @@ float FSceneViewport::GetMaximumFullFrameLuminanceInNits() const
 	return MaximumFullFrameLuminanceInNits;
 }
 
+float FSceneViewport::GetHDRPaperWhiteInNits() const
+{
+	return HDRPaperWhiteInNits;
+}
+
 const UE::Color::FColorSpace& FSceneViewport::GetLimitingColorSpace() const
 {
 	return LimitingColorSpace;
@@ -1977,6 +1985,7 @@ void FSceneViewport::UpdateViewportRHI(bool bDestroyed, uint32 NewSizeX, uint32
 			MinimumLuminanceInNits = OutHDRMetaData.MinimumLuminanceInNits;
 			MaximumLuminanceInNits = OutHDRMetaData.MaximumLuminanceInNits;
 			MaximumFullFrameLuminanceInNits = OutHDRMetaData.MaximumFullFrameLuminanceInNits;
+			HDRPaperWhiteInNits = OutHDRMetaData.HDRPaperWhiteInNits;
 			LimitingColorSpace = OutHDRMetaData.LimitingColorSpace;
 			bHDRViewport = OutHDRMetaData.bHDRSupported;
 
diff --git a/Engine/Source/Runtime/Engine/Public/Slate/SceneViewport.h b/Engine/Source/Runtime/Engine/Public/Slate/SceneViewport.h
--- a/Engine/Source/Runtime/Engine/Public/Slate/SceneViewport.h
+++ b/Engine/Source/Runtime/Engine/Public/Slate/SceneViewport.h
@@ -264,6 +264,7 @@ public:
 	ENGINE_API virtual float GetMinimumLuminanceInNits() const override;
 	ENGINE_API virtual float GetMaximumLuminanceInNits() const override;
 	ENGINE_API virtual float GetMaximumFullFrameLuminanceInNits() const override;
+	ENGINE_API virtual float GetHDRPaperWhiteInNits() const override;
 	ENGINE_API virtual const UE::Color::FColorSpace& GetLimitingColorSpace() const override;
 	ENGINE_API virtual ESlateViewportDynamicRange GetViewportDynamicRange() const override;
 
@@ -482,6 +483,7 @@ private:
 	float MinimumLuminanceInNits;
 	float MaximumLuminanceInNits;
 	float MaximumFullFrameLuminanceInNits;
+	float HDRPaperWhiteInNits;
 	UE::Color::FColorSpace LimitingColorSpace;
 
 	/** The optional gamma value to use for this viewport */
diff --git a/Engine/Source/Runtime/Engine/Public/UnrealClient.h b/Engine/Source/Runtime/Engine/Public/UnrealClient.h
--- a/Engine/Source/Runtime/Engine/Public/UnrealClient.h
+++ b/Engine/Source/Runtime/Engine/Public/UnrealClient.h
@@ -76,6 +76,7 @@ public:
 	virtual float GetMinimumLuminanceInNits() const { return 0.f; }
 	virtual float GetMaximumLuminanceInNits() const { return 100.f; }
 	virtual float GetMaximumFullFrameLuminanceInNits() const { return 100.f; }
+	virtual float GetHDRPaperWhiteInNits() const { return 203.f; }
 	virtual const UE::Color::FColorSpace& GetLimitingColorSpace() const { return UE::Color::FColorSpace::GetSRGB(); }
 
 	/**
diff --git a/Engine/Source/Runtime/RenderCore/Private/RenderCore.cpp b/Engine/Source/Runtime/RenderCore/Private/RenderCore.cpp
--- a/Engine/Source/Runtime/RenderCore/Private/RenderCore.cpp
+++ b/Engine/Source/Runtime/RenderCore/Private/RenderCore.cpp
@@ -514,6 +514,7 @@ inline FHDRMetaData HDRGetDefaultMetaData()
 		.MaximumFullFrameLuminanceInNits = bHDRSupported ?
 			static_cast<float>(CVarHDRDisplayMaxLuminance.GetValueOnAnyThread()) :
 			100.f,
+		.HDRPaperWhiteInNits = 203.f,
 		.LimitingColorSpace = bHDRSupported ? UE::Color::FColorSpace::GetP3D65() : UE::Color::FColorSpace::GetSRGB(),
 		.bHDRSupported = bHDRSupported,
 	};
@@ -552,6 +553,7 @@ RENDERCORE_API void HDRAddCustomMetaData(void* OSWindow, EDisplayOutputFormat Di
 	HDRMetaData.MaximumFullFrameLuminanceInNits = bHDREnabled ?
 		CVarHDRDisplayMaxLuminance.GetValueOnAnyThread() :
 		100.f;
+	HDRMetaData.HDRPaperWhiteInNits = 203.f;
 	HDRMetaData.bHDRSupported = bHDREnabled;
 }
 
@@ -643,6 +645,7 @@ RENDERCORE_API void HDRGetMetaDataForWindow(FHDRMetaData& OutHDRMetaData, const
 	OutHDRMetaData.MinimumLuminanceInNits = OutHDRMetaData.bHDRSupported ? DisplayList[BestDisplayForWindow].MinimumLuminanceInNits : 0.f;
 	OutHDRMetaData.MaximumLuminanceInNits = OutHDRMetaData.bHDRSupported ? DisplayList[BestDisplayForWindow].MaximumLuminanceInNits : 100.f;
 	OutHDRMetaData.MaximumFullFrameLuminanceInNits = OutHDRMetaData.bHDRSupported ? DisplayList[BestDisplayForWindow].MaximumFullFrameLuminanceInNits : 100.f;
+	OutHDRMetaData.HDRPaperWhiteInNits = DisplayList[BestDisplayForWindow].HDRPaperWhiteInNits;
 	OutHDRMetaData.LimitingColorSpace = OutHDRMetaData.bHDRSupported ? DisplayList[BestDisplayForWindow].LimitingColorSpace : UE::Color::FColorSpace::GetSRGB();
 
 	if (OutHDRMetaData.bHDRSupported)
diff --git a/Engine/Source/Runtime/RenderCore/Public/HDRHelper.h b/Engine/Source/Runtime/RenderCore/Public/HDRHelper.h
--- a/Engine/Source/Runtime/RenderCore/Public/HDRHelper.h
+++ b/Engine/Source/Runtime/RenderCore/Public/HDRHelper.h
@@ -20,6 +20,7 @@ struct FHDRMetaData
 	float MinimumLuminanceInNits;
 	float MaximumLuminanceInNits;
 	float MaximumFullFrameLuminanceInNits;
+	float HDRPaperWhiteInNits;
 	UE::Color::FColorSpace LimitingColorSpace;
 	bool bHDRSupported;
 };
@@ -42,6 +43,7 @@ struct FDisplayInformation
 	float MinimumLuminanceInNits = 0.f;
 	float MaximumLuminanceInNits = 100.f;
 	float MaximumFullFrameLuminanceInNits = 100.f;
+	float HDRPaperWhiteInNits = 203.f;
 	UE::Color::FColorSpace LimitingColorSpace = UE::Color::FColorSpace::GetSRGB();
 };
 
diff --git a/Engine/Source/Runtime/Renderer/Private/PostProcess/PostProcessCombineLUTs.cpp b/Engine/Source/Runtime/Renderer/Private/PostProcess/PostProcessCombineLUTs.cpp
--- a/Engine/Source/Runtime/Renderer/Private/PostProcess/PostProcessCombineLUTs.cpp
+++ b/Engine/Source/Runtime/Renderer/Private/PostProcess/PostProcessCombineLUTs.cpp
@@ -191,7 +191,8 @@ BEGIN_SHADER_PARAMETER_STRUCT(FCombineLUTParameters, )
 	SHADER_PARAMETER(float, FilmShoulder)
 	SHADER_PARAMETER(float, FilmBlackClip)
 	SHADER_PARAMETER(float, FilmWhiteClip)
-	SHADER_PARAMETER(uint32, bUseMobileTonemapper)
+	SHADER_PARAMETER(float, HDRPaperWhiteRatio)
+	SHADER_PARAMETER(float, InvHDRPaperWhiteRatio)
 	SHADER_PARAMETER(uint32, bIsTemperatureWhiteBalance)
 	SHADER_PARAMETER_STRUCT_INCLUDE(FColorRemapParameters, ColorRemap)
 	SHADER_PARAMETER_STRUCT_INCLUDE(FTonemapperOutputDeviceParameters, OutputDevice)
@@ -312,6 +313,10 @@ struct FCachedLUTSettings
 		UPDATE_CACHE_SETTINGS(Parameters.LimitingXYZToRgb1, FVector3f(LimitingXYZToRgb.GetColumn(1)), bHasChanged);
 		UPDATE_CACHE_SETTINGS(Parameters.LimitingXYZToRgb2, FVector3f(LimitingXYZToRgb.GetColumn(2)), bHasChanged);
 
+		const float HDRPaperWhiteRatio = ViewFamily.RenderTarget->GetHDRPaperWhiteInNits() / 203.f;
+		UPDATE_CACHE_SETTINGS(Parameters.HDRPaperWhiteRatio, HDRPaperWhiteRatio, bHasChanged);
+		UPDATE_CACHE_SETTINGS(Parameters.InvHDRPaperWhiteRatio, 1.f / HDRPaperWhiteRatio, bHasChanged);
+
 		FACESTonemapParams TonemapperParams;
 		GetACESTonemapParameters(TonemapperParams);
 		UPDATE_CACHE_SETTINGS(Parameters.ACESTonemapParameters.ACESMinMaxData, TonemapperParams.ACESMinMaxData, bHasChanged);
diff --git a/Engine/Source/Runtime/Renderer/Private/PostProcess/PostProcessSelectionOutline.cpp b/Engine/Source/Runtime/Renderer/Private/PostProcess/PostProcessSelectionOutline.cpp
--- a/Engine/Source/Runtime/Renderer/Private/PostProcess/PostProcessSelectionOutline.cpp
+++ b/Engine/Source/Runtime/Renderer/Private/PostProcess/PostProcessSelectionOutline.cpp
@@ -379,8 +379,20 @@ FScreenPassTexture AddSelectionOutlinePass(
 			}
 		}
 
-		static const auto CVarHDRUILuminance = IConsoleManager::Get().FindTConsoleVariableDataFloat(TEXT("r.HDR.UI.Luminance"));
-		float UILuminance = CVarHDRUILuminance ? CVarHDRUILuminance->GetValueOnRenderThread() : 300.0f;
+		float UILuminance = 203.f;
+		static const auto* const CVarHDRUIUseOSPaperWhite = IConsoleManager::Get().FindTConsoleVariableDataBool(TEXT("r.HDR.UI.UseOSPaperWhite"));
+		if (CVarHDRUIUseOSPaperWhite && CVarHDRUIUseOSPaperWhite->GetValueOnRenderThread())
+		{
+			UILuminance = View.Family->RenderTarget->GetHDRPaperWhiteInNits();
+		}
+		else
+		{
+			static const auto* const CVarHDRUILuminance = IConsoleManager::Get().FindTConsoleVariableDataFloat(TEXT("r.HDR.UI.Luminance"));
+			if (CVarHDRUILuminance)
+			{
+				UILuminance = CVarHDRUILuminance->GetValueOnRenderThread();
+			}
+		}
 		PassParameters->UILuminanceAndIsSCRGB = bIsSCRGB ? UILuminance : -UILuminance;
 
 		FSelectionOutlinePS::FPermutationDomain PermutationVector;
diff --git a/Engine/Source/Runtime/SlateRHIRenderer/Private/SlatePostProcessor.cpp b/Engine/Source/Runtime/SlateRHIRenderer/Private/SlatePostProcessor.cpp
--- a/Engine/Source/Runtime/SlateRHIRenderer/Private/SlatePostProcessor.cpp
+++ b/Engine/Source/Runtime/SlateRHIRenderer/Private/SlatePostProcessor.cpp
@@ -31,7 +31,13 @@ float GetSlateHDRUILevel()
 float GetSlateHDRUILuminance()
 {
 	static auto CVar = IConsoleManager::Get().FindConsoleVariable(TEXT("r.HDR.UI.Luminance"));
-	return CVar ? CVar->GetFloat() : 300.0f;
+	return CVar ? CVar->GetFloat() : 203.f;
+}
+
+bool GetSlateHDRUIUseOSPaperWhite()
+{
+	static auto CVar = IConsoleManager::Get().FindConsoleVariable(TEXT("r.HDR.UI.UseOSPaperWhite"));
+	return CVar ? CVar->GetBool() : false;
 }
 
 int32 GetSlateHDRUICompositeEOTF()
@@ -43,7 +49,7 @@ int32 GetSlateHDRUICompositeEOTF()
 float GetSlateHDRUICompositeEOTFGamma()
 {
 	static auto CVar = IConsoleManager::Get().FindConsoleVariable(TEXT("r.HDR.UI.CompositeEOTF.Gamma"));
-	return CVar ? CVar->GetFloat() : 2.2f;
+	return CVar ? CVar->GetFloat() : 2.4f;
 }
 
 int32 GetSlateHDRUICompositeBlend()
@@ -55,7 +61,7 @@ int32 GetSlateHDRUICompositeBlend()
 float GetSlateHDRUICompositeBlendGamma()
 {
 	static auto CVar = IConsoleManager::Get().FindConsoleVariable(TEXT("r.HDR.UI.CompositeBlend.Gamma"));
-	return CVar ? CVar->GetFloat() : 2.3f;
+	return CVar ? CVar->GetFloat() : 2.4f;
 }
 
 ETextureCreateFlags GetSlateTransientRenderTargetFlags()
@@ -121,6 +127,7 @@ struct FSlateCompositeHDRForBlurPassInputs
 	FRDGTexture* InputCompositeTexture;
 	FRDGTexture* InputTexture;
 	FIntPoint OutputExtent;
+	float HDRPaperWhiteInNits;
 };
 
 FScreenPassTexture AddSlateCompositeHDRForBlurPass(FRDGBuilder& GraphBuilder, const FSlateCompositeHDRForBlurPassInputs& Inputs)
@@ -154,7 +161,7 @@ FScreenPassTexture AddSlateCompositeHDRForBlurPass(FRDGBuilder& GraphBuilder, co
 	PassParameters->UISampler = TStaticSamplerState<SF_Bilinear>::GetRHI();
 	PassParameters->UITextureSize = InputViewport.Extent;
 	PassParameters->UILevel = GetSlateHDRUILevel();
-	PassParameters->UILuminance = GetSlateHDRUILuminance();
+	PassParameters->UILuminance = GetSlateHDRUIUseOSPaperWhite() ? Inputs.HDRPaperWhiteInNits : GetSlateHDRUILuminance();
 	PassParameters->UICompositeEOTF = GetSlateHDRUICompositeEOTF();
 	PassParameters->UICompositeEOTFGamma = GetSlateHDRUICompositeEOTFGamma();
 	PassParameters->UICompositeBlend = GetSlateHDRUICompositeBlend();
@@ -761,6 +768,7 @@ void AddSlatePostProcessBlurPass(FRDGBuilder& GraphBuilder, const FSlatePostProc
 		CompositeInputs.InputTexture = Inputs.InputTexture;
 		CompositeInputs.InputCompositeTexture = Inputs.SDRCompositeUITexture;
 		CompositeInputs.OutputExtent = OutputTextureViewport.Extent;
+		CompositeInputs.HDRPaperWhiteInNits = Inputs.HDRPaperWhiteInNits;
 
 		BlurInputTexture = AddSlateCompositeHDRForBlurPass(GraphBuilder, CompositeInputs);
 	}
diff --git a/Engine/Source/Runtime/SlateRHIRenderer/Private/SlatePostProcessor.h b/Engine/Source/Runtime/SlateRHIRenderer/Private/SlatePostProcessor.h
--- a/Engine/Source/Runtime/SlateRHIRenderer/Private/SlatePostProcessor.h
+++ b/Engine/Source/Runtime/SlateRHIRenderer/Private/SlatePostProcessor.h
@@ -69,6 +69,7 @@ struct FSlatePostProcessBlurPassInputs
 	float  Strength = 0.0f;
 	uint32 DownsampleAmount = 0;
 	FVector4f CornerRadius = FVector4f::Zero();
+	float HDRPaperWhiteInNits = 203.f;
 };
 
 extern ETextureCreateFlags GetSlateTransientRenderTargetFlags();
diff --git a/Engine/Source/Runtime/SlateRHIRenderer/Private/SlateRHIRenderer.cpp b/Engine/Source/Runtime/SlateRHIRenderer/Private/SlateRHIRenderer.cpp
--- a/Engine/Source/Runtime/SlateRHIRenderer/Private/SlateRHIRenderer.cpp
+++ b/Engine/Source/Runtime/SlateRHIRenderer/Private/SlateRHIRenderer.cpp
@@ -88,6 +88,16 @@ static TAutoConsoleVariable<float> CVarHDRUILuminance(
 	TEXT("Base Luminance in nits for UI elements when compositing into HDR framebuffer. Gets multiplied by r.HDR.UI.Level"),
 	ECVF_RenderThreadSafe);
 
+static TAutoConsoleVariable<bool> CVarHDRUIUseOSPaperWhite(
+	TEXT("r.HDR.UI.UseOSPaperWhite"),
+#if PLATFORM_PROVIDES_HDR_PAPER_WHITE
+	true,
+#else
+	false,
+#endif
+	TEXT("Whether to use the paper white value provided by the OS for UI luminance."),
+	ECVF_RenderThreadSafe);
+
 static TAutoConsoleVariable<int32> CVarHDRUICompositeEOTF(
 	TEXT("r.HDR.UI.CompositeEOTF"),
 	0,
@@ -211,6 +221,7 @@ struct FSlateViewportInfo : public FRenderResource
 	EPixelFormat PixelFormat = EPixelFormat::PF_Unknown;
 	EDisplayColorGamut HDRDisplayColorGamut = EDisplayColorGamut::sRGB_D65;
 	EDisplayOutputFormat HDRDisplayOutputFormat = EDisplayOutputFormat::SDR_sRGB;
+	float HDRPaperWhiteInNits = 203.f;
 	bool bDisplayFormatIsHDR = false;
 	bool bFullscreen = false;
 
@@ -404,7 +415,13 @@ void FSlateRHIRenderer::CreateViewport(const TSharedRef<SWindow> Window)
 	ViewInfo->ProjectionMatrix = CreateSlateProjectionMatrix(ExtentToResizeTo.X, ExtentToResizeTo.Y);
 	ViewInfo->Extent = ExtentToResizeTo;
 	ViewInfo->ExtentToResizeTo = ExtentToResizeTo;
-	HDRGetMetaData(ViewInfo->HDRDisplayOutputFormat, ViewInfo->HDRDisplayColorGamut, ViewInfo->bDisplayFormatIsHDR, Window->GetPositionInScreen(), Window->GetPositionInScreen() + Window->GetSizeInScreen(), ViewInfo->OSWindow);
+
+	FHDRMetaData OutHDRMetaData;
+	HDRGetMetaDataForWindow(OutHDRMetaData, Window->GetPositionInScreen(), Window->GetPositionInScreen() + Window->GetSizeInScreen(), ViewInfo->OSWindow);
+	ViewInfo->HDRDisplayOutputFormat = OutHDRMetaData.DisplayOutputFormat;
+	ViewInfo->HDRDisplayColorGamut = OutHDRMetaData.DisplayColorGamut;
+	ViewInfo->HDRPaperWhiteInNits = OutHDRMetaData.HDRPaperWhiteInNits;
+	ViewInfo->bDisplayFormatIsHDR = OutHDRMetaData.bHDRSupported;
 
 	const bool bFullscreen = IsViewportFullscreen(*Window);
 	ViewInfo->PixelFormat = GetViewportPixelFormat(*Window, ViewInfo->bDisplayFormatIsHDR);
@@ -432,15 +449,21 @@ void FSlateRHIRenderer::ResizeViewportIfNeeded(FSlateViewportInfo* ViewInfo, FIn
 		return;
 	}
 
-	bool bHDREnabled = IsHDREnabled();
-	EDisplayColorGamut HDRColorGamut = HDRGetDefaultDisplayColorGamut();
-	EDisplayOutputFormat HDROutputDevice = HDRGetDefaultDisplayOutputFormat();
-	HDRGetMetaData(HDROutputDevice, HDRColorGamut, bHDREnabled, Window->GetPositionInScreen(), Window->GetPositionInScreen() + Window->GetSizeInScreen(), ViewInfo->OSWindow);
+	FHDRMetaData OutHDRMetaData;
+	OutHDRMetaData.DisplayOutputFormat = HDRGetDefaultDisplayOutputFormat();
+	OutHDRMetaData.DisplayColorGamut = HDRGetDefaultDisplayColorGamut();
+	OutHDRMetaData.bHDRSupported = false;
+	HDRGetMetaDataForWindow(OutHDRMetaData, Window->GetPositionInScreen(), Window->GetPositionInScreen() + Window->GetSizeInScreen(), ViewInfo->OSWindow);
+	const EDisplayOutputFormat HDROutputDevice = OutHDRMetaData.DisplayOutputFormat;
+	const EDisplayColorGamut HDRColorGamut = OutHDRMetaData.DisplayColorGamut;
+	const float HDRPaperWhiteInNits = OutHDRMetaData.HDRPaperWhiteInNits;
+	const bool bHDREnabled = OutHDRMetaData.bHDRSupported;
 
 	bool bHDRStale = false;
-	bHDRStale |= HDROutputDevice != ViewInfo->HDRDisplayOutputFormat;
-	bHDRStale |= HDRColorGamut   != ViewInfo->HDRDisplayColorGamut;
-	bHDRStale |= bHDREnabled     != ViewInfo->bDisplayFormatIsHDR;
+	bHDRStale |= HDROutputDevice     != ViewInfo->HDRDisplayOutputFormat;
+	bHDRStale |= HDRColorGamut       != ViewInfo->HDRDisplayColorGamut;
+	bHDRStale |= HDRPaperWhiteInNits != ViewInfo->HDRPaperWhiteInNits;
+	bHDRStale |= bHDREnabled         != ViewInfo->bDisplayFormatIsHDR;
 
 	if (bHDRStale || ViewInfo->Extent != ExtentToResizeTo || ViewInfo->bFullscreen != bFullscreen || !IsValidRef(ViewInfo->ViewportRHI))
 	{
@@ -481,6 +504,7 @@ void FSlateRHIRenderer::ResizeViewportIfNeeded(FSlateViewportInfo* ViewInfo, FIn
 		ViewInfo->bFullscreen = bFullscreen;
 		ViewInfo->HDRDisplayColorGamut = HDRColorGamut;
 		ViewInfo->HDRDisplayOutputFormat = HDROutputDevice;
+		ViewInfo->HDRPaperWhiteInNits = HDRPaperWhiteInNits;
 		ViewInfo->bDisplayFormatIsHDR = bHDREnabled;
 		ViewInfo->PixelFormat = GetViewportPixelFormat(*Window, bHDREnabled);
 
@@ -784,10 +808,11 @@ void FSlateRHIRenderer::DrawWindowViewport_RenderThread(ISlateViewport* SlateVie
 		, .CursorPosition               = Inputs.CursorPosition
 		, .Time                         = Inputs.Time
 		, .HDRDisplayColorGamut         = ViewportInfo.HDRDisplayColorGamut
+		, .HDRPaperWhiteInNits          = ViewportInfo.HDRPaperWhiteInNits
 		, .UsedSlatePostBuffers         = Inputs.UsedSlatePostBuffers
 		, .ViewportScaleUI              = Inputs.ViewportScaleUI
 		, .bWireframe                   = GSlateWireframe
-		, .bElementsTextureIsHDRDisplay =  ViewportInfo.bDisplayFormatIsHDR
+		, .bElementsTextureIsHDRDisplay = ViewportInfo.bDisplayFormatIsHDR
 	};
 	
 	if (ScreenshotState.ViewportToCapture == &ViewportInfo)
@@ -849,7 +874,9 @@ void FSlateRHIRenderer::DrawWindowViewport_RenderThread(ISlateViewport* SlateVie
 		CommonParameters.UITexture = ElementsTexture;
 		CommonParameters.UISampler = TStaticSamplerState<SF_Point>::GetRHI();
 		CommonParameters.UILevel = CVarUILevel.GetValueOnRenderThread();
-		CommonParameters.UILuminance = CVarHDRUILuminance.GetValueOnRenderThread();
+		CommonParameters.UILuminance = CVarHDRUIUseOSPaperWhite.GetValueOnRenderThread()
+			? ViewportInfo.HDRPaperWhiteInNits
+			: CVarHDRUILuminance.GetValueOnRenderThread();
 		CommonParameters.UICompositeEOTF = CVarHDRUICompositeEOTF.GetValueOnRenderThread();
 		CommonParameters.UICompositeEOTFGamma = CVarHDRUICompositeEOTFGamma.GetValueOnRenderThread();
 		CommonParameters.UICompositeBlend = CVarHDRUICompositeBlend.GetValueOnRenderThread();
diff --git a/Engine/Source/Runtime/SlateRHIRenderer/Private/SlateRHIRenderingPolicy.cpp b/Engine/Source/Runtime/SlateRHIRenderer/Private/SlateRHIRenderingPolicy.cpp
--- a/Engine/Source/Runtime/SlateRHIRenderer/Private/SlateRHIRenderingPolicy.cpp
+++ b/Engine/Source/Runtime/SlateRHIRenderer/Private/SlateRHIRenderingPolicy.cpp
@@ -1714,6 +1714,7 @@ void AddSlateDrawElementsPass(
 			BlurInputs.Strength = ShaderParams.PixelParams2.Y;
 			BlurInputs.DownsampleAmount = ShaderParams.PixelParams2.Z;
 			BlurInputs.CornerRadius = ShaderParams.PixelParams3;
+			BlurInputs.HDRPaperWhiteInNits = Inputs.HDRPaperWhiteInNits;
 
 			AddSlatePostProcessBlurPass(GraphBuilder, BlurInputs);
 			break;
diff --git a/Engine/Source/Runtime/SlateRHIRenderer/Private/SlateRHIRenderingPolicy.h b/Engine/Source/Runtime/SlateRHIRenderer/Private/SlateRHIRenderingPolicy.h
--- a/Engine/Source/Runtime/SlateRHIRenderer/Private/SlateRHIRenderingPolicy.h
+++ b/Engine/Source/Runtime/SlateRHIRenderer/Private/SlateRHIRenderingPolicy.h
@@ -66,6 +66,7 @@ struct FSlateDrawElementsPassInputs
 	FIntPoint CursorPosition = FIntPoint::ZeroValue;
 	FGameTime Time;
 	EDisplayColorGamut HDRDisplayColorGamut = EDisplayColorGamut::sRGB_D65;
+	float HDRPaperWhiteInNits = 203.f;
 	ESlatePostRT UsedSlatePostBuffers = ESlatePostRT::None;
 	float ViewportScaleUI = 1.0f;
 	bool bWireframe = false;
