Fix assert when booting with HDR enabled
- Prior to CL 51012639, Slate would skip resizing the viewport if the caller is not the game thread.
- Restoring this check prevents the Slate loading thread from resizing the viewport and causing check(IsInGameThread()) asserts.

#Synced-CL 51352388

diff --git a/Engine/Source/Runtime/SlateRHIRenderer/Private/SlateRHIRenderer.cpp b/Engine/Source/Runtime/SlateRHIRenderer/Private/SlateRHIRenderer.cpp
--- a/Engine/Source/Runtime/SlateRHIRenderer/Private/SlateRHIRenderer.cpp
+++ b/Engine/Source/Runtime/SlateRHIRenderer/Private/SlateRHIRenderer.cpp
@@ -513,6 +513,10 @@ bool FSlateViewportInfo::IsDisplayFormatHDR() const
 void FSlateViewportInfo::ResizeViewport(FIntPoint ExtentToResizeTo, bool bInFullscreen)
 {
 	check(IsThreadSafeForSlateRendering());
+	if (!IsInGameThread() || IsInSlateThread())
+	{
+		return;
+	}
 
 	ExtentToResizeTo = ClampExtent(ExtentToResizeTo);
 
