Fixed: Missing texture binding in a slate material elements
#Synced-CL 52121382

diff --git a/Engine/Source/Runtime/SlateRHIRenderer/Private/SlateRHIRenderingPolicy.cpp b/Engine/Source/Runtime/SlateRHIRenderer/Private/SlateRHIRenderingPolicy.cpp
--- a/Engine/Source/Runtime/SlateRHIRenderer/Private/SlateRHIRenderingPolicy.cpp
+++ b/Engine/Source/Runtime/SlateRHIRenderer/Private/SlateRHIRenderingPolicy.cpp
@@ -1194,6 +1194,10 @@ FSlateRenderBatchOp* CreateSlateRenderBatchOp(
 		{
 			PixelShader->SetAdditionalTexture(*PixelBindings, MaskResource->GetTypedResource(), TStaticSamplerState<SF_Bilinear, AM_Clamp, AM_Clamp, AM_Clamp>::GetRHI());
 		}
+		else
+		{
+			PixelShader->SetAdditionalTexture(*PixelBindings, GBlackTexture->TextureRHI, TStaticSamplerState<SF_Bilinear, AM_Clamp, AM_Clamp, AM_Clamp>::GetRHI());
+		}
 
 		BlendState = GetMaterialBlendState(MaskResource, EffectiveMaterial);
 	}
