When running with a null renderer, SlateRHIRenderer is not loaded, leading to errors in BP, as some required types can't be found.
=> extracted blueprintable content from SlateRHIRenderer into (the new) SlateBaseRenderer

#tests Engine test now have an automated test using post buffer blur. Ran with null renderer and validated that there was no error with BP (they are still executed properly). Validated that all moved BP functions/classes are deprecated properly by creating first data using it in engine test, then reopening it with new code.
#Synced-CL 52006545
#Synced-CL 52006545

diff --git a/Engine/Config/BaseEngine.ini b/Engine/Config/BaseEngine.ini
--- a/Engine/Config/BaseEngine.ini
+++ b/Engine/Config/BaseEngine.ini
@@ -1413,6 +1413,11 @@ ClassName=/Script/OnlineSubsystemUtils.OnlineEngineInterfaceImpl
 +StructRedirects=(OldName="/Script/Water.ShallowWaterSimulationGrid",NewName="/Script/Water.ShallowWaterSimulationGrid_DEPRECATED")
 +PropertyRedirects=(OldName="/Script/Water.BakedShallowWaterSimulationComponent.SimulationData",NewName="/Script/Water.BakedShallowWaterSimulationComponent.SimulationData_DEPRECATED")
 
++ClassRedirects=(OldName="/Script/SlateRHIRenderer.SlateRHIPostBufferProcessor",NewName="/Script/SlateBaseRenderer.SlatePostBufferProcessor")
++ClassRedirects=(OldName="/Script/SlateRHIRenderer.SlatePostBufferBlur",NewName="/Script/SlateBaseRenderer.SlatePostBufferBlur")
++ClassRedirects=(OldName="/Script/SlateRHIRenderer.SlateFXSubsystem",NewName="/Script/SlateBaseRenderer.SlateFXBaseSubsystem",bInstanceOnly=false)
++FunctionRedirects=(OldName="/Script/SlateRHIRenderer.SlateFXSubsystem.GetSlatePostProcessor",NewName="/Script/SlateBaseRenderer.SlateFXBaseSubsystem.GetSlatePostProcessor")
+
 ; End of releases redirects
 
 ; Fixing the output of tags for blueprints to be reset.
diff --git a/Engine/Source/Runtime/SlateBaseRenderer/Private/FX/SlateFXBaseSubsystem.cpp b/Engine/Source/Runtime/SlateBaseRenderer/Private/FX/SlateFXBaseSubsystem.cpp
new file mode 100644
--- /dev/null
+++ b/Engine/Source/Runtime/SlateBaseRenderer/Private/FX/SlateFXBaseSubsystem.cpp
@@ -0,0 +1,81 @@
+// Copyright Epic Games, Inc. All Rights Reserved.
+
+#include "FX/SlateFXBaseSubsystem.h"
+
+#include "Engine/Engine.h"
+
+#include UE_INLINE_GENERATED_CPP_BY_NAME(SlateFXBaseSubsystem)
+
+USlatePostBufferProcessor* USlateFXBaseSubsystem::GetPostProcessor(ESlatePostRT InSlatePostBufferBit)
+{
+	USlatePostBufferProcessor* Result = nullptr;
+
+	if (GEngine)
+	{
+		if (USlateFXBaseSubsystem* SlateFXSubsystem = GEngine->GetEngineSubsystem<USlateFXBaseSubsystem>())
+		{
+			Result = SlateFXSubsystem->GetSlatePostProcessor(InSlatePostBufferBit);
+		}
+	}
+
+	return Result;
+}
+
+bool USlateFXBaseSubsystem::ShouldCreateSubsystem(UObject* Outer) const
+{
+	TArray<UClass*> ChildClasses;
+	GetDerivedClasses(GetClass(), ChildClasses, false);
+
+	// Only create an instance if there is no override implementation defined elsewhere
+	return ChildClasses.Num() == 0;
+}
+
+void USlateFXBaseSubsystem::Initialize(FSubsystemCollectionBase& Collection)
+{
+	Super::Initialize(Collection);
+
+	FWorldDelegates::OnPreWorldInitialization.AddUObject(this, &USlateFXBaseSubsystem::OnPreWorldInitialization);
+	FWorldDelegates::OnPostWorldCleanup.AddUObject(this, &USlateFXBaseSubsystem::OnPostWorldCleanup);
+}
+
+void USlateFXBaseSubsystem::Deinitialize()
+{
+	SlatePostBufferProcessors.Empty();
+
+	FWorldDelegates::OnPreWorldInitialization.RemoveAll(this);
+	FWorldDelegates::OnPostWorldCleanup.RemoveAll(this);
+
+	Super::Deinitialize();
+}
+
+USlatePostBufferProcessor* USlateFXBaseSubsystem::GetSlatePostProcessor(ESlatePostRT InPostBufferBit)
+{
+	if (TObjectPtr<USlatePostBufferProcessor>* Processor = SlatePostBufferProcessors.Find(InPostBufferBit))
+	{
+		return *Processor;
+	}
+
+	return nullptr;
+}
+
+void USlateFXBaseSubsystem::OnPreWorldInitialization(UWorld* World, const UWorld::InitializationValues IVS)
+{
+	if (World && World->WorldType == EWorldType::EditorPreview)
+	{
+		return;
+	}
+
+	SlatePostBufferProcessors.Empty();
+	OnInitProcessors();
+}
+
+void USlateFXBaseSubsystem::OnPostWorldCleanup(UWorld* World, bool SessionEnded, bool bCleanupResources)
+{
+	if (World && World->WorldType == EWorldType::EditorPreview)
+	{
+		return;
+	}
+
+	SlatePostBufferProcessors.Empty();
+	OnCleanupProcessors();
+}
diff --git a/Engine/Source/Runtime/SlateBaseRenderer/Private/FX/SlatePostBufferBlur.cpp b/Engine/Source/Runtime/SlateBaseRenderer/Private/FX/SlatePostBufferBlur.cpp
new file mode 100644
--- /dev/null
+++ b/Engine/Source/Runtime/SlateBaseRenderer/Private/FX/SlatePostBufferBlur.cpp
@@ -0,0 +1,6 @@
+// Copyright Epic Games, Inc. All Rights Reserved.
+
+#include "FX/SlatePostBufferBlur.h"
+
+#include UE_INLINE_GENERATED_CPP_BY_NAME(SlatePostBufferBlur)
+
diff --git a/Engine/Source/Runtime/SlateBaseRenderer/Private/FX/SlatePostBufferProcessor.cpp b/Engine/Source/Runtime/SlateBaseRenderer/Private/FX/SlatePostBufferProcessor.cpp
new file mode 100644
--- /dev/null
+++ b/Engine/Source/Runtime/SlateBaseRenderer/Private/FX/SlatePostBufferProcessor.cpp
@@ -0,0 +1,5 @@
+// Copyright Epic Games, Inc. All Rights Reserved.
+
+#include "FX/SlatePostBufferProcessor.h"
+
+#include UE_INLINE_GENERATED_CPP_BY_NAME(SlatePostBufferProcessor)
\ No newline at end of file
diff --git a/Engine/Source/Runtime/SlateBaseRenderer/Private/SlateBaseRendererModule.cpp b/Engine/Source/Runtime/SlateBaseRenderer/Private/SlateBaseRendererModule.cpp
new file mode 100644
--- /dev/null
+++ b/Engine/Source/Runtime/SlateBaseRenderer/Private/SlateBaseRendererModule.cpp
@@ -0,0 +1,5 @@
+// Copyright Epic Games, Inc. All Rights Reserved.
+
+#include "Modules/ModuleManager.h"
+
+IMPLEMENT_MODULE(FDefaultModuleImpl, SlateBaseRenderer)
diff --git a/Engine/Source/Runtime/SlateBaseRenderer/Private/SlateFontInfoBlueprintLibrary.cpp b/Engine/Source/Runtime/SlateBaseRenderer/Private/SlateFontInfoBlueprintLibrary.cpp
new file mode 100644
--- /dev/null
+++ b/Engine/Source/Runtime/SlateBaseRenderer/Private/SlateFontInfoBlueprintLibrary.cpp
@@ -0,0 +1,44 @@
+// Copyright Epic Games, Inc. All Rights Reserved.
+
+#include "SlateFontInfoBlueprintLibrary.h"
+#include "Blueprint/BlueprintExceptionInfo.h"
+#include "Materials/MaterialInterface.h"
+
+#include UE_INLINE_GENERATED_CPP_BY_NAME(SlateFontInfoBlueprintLibrary)
+
+DEFINE_LOG_CATEGORY(LogSlateFontInfoBlueprint);
+
+
+FSlateFontInfo USlateFontInfoBlueprintLibrary::MakeSlateFontInfo(const UObject* FontObject, UObject* FontMaterial, FFontOutlineSettings OutlineSettings, FName TypefaceFontName, float Size, int32 LetterSpacing, float SkewAmount, bool bForceMonospaced, bool bMaterialIsStencil, float MonospacedWidth)
+{
+	FSlateFontInfo SlateFontInfo = FSlateFontInfo();
+	SlateFontInfo.FontObject = FontObject;
+	SlateFontInfo.FontMaterial = FontMaterial;
+	SlateFontInfo.OutlineSettings = OutlineSettings;
+	SlateFontInfo.TypefaceFontName = TypefaceFontName;
+	SlateFontInfo.Size = Size;
+	SlateFontInfo.LetterSpacing = LetterSpacing;
+	SlateFontInfo.SkewAmount = SkewAmount;
+	SlateFontInfo.bForceMonospaced = bForceMonospaced;
+	SlateFontInfo.bMaterialIsStencil = bMaterialIsStencil;
+	SlateFontInfo.MonospacedWidth = MonospacedWidth;
+
+	if (FontMaterial && !FontMaterial->IsA<UMaterialInterface>())
+	{
+#if !UE_BUILD_SHIPPING && !UE_BUILD_TEST
+		if (FFrame::GetThreadLocalTopStackFrame() && FFrame::GetThreadLocalTopStackFrame()->Object)
+		{
+			FText ErrorMessage = FText::FromString(TEXT("Font Material is not of type material. Either remove the reference or replace it with an object of type material."));
+			const FBlueprintExceptionInfo ExceptionInfo(EBlueprintExceptionType::FatalError, ErrorMessage);
+			FBlueprintCoreDelegates::ThrowScriptException(FFrame::GetThreadLocalTopStackFrame()->Object, *FFrame::GetThreadLocalTopStackFrame(), ExceptionInfo);
+		}
+		else
+#endif
+		{
+			UE_LOG(LogSlateFontInfoBlueprint, Warning, TEXT("Font Material is not of type material. It will be replaced with no material."));
+		}
+		SlateFontInfo.FontMaterial = nullptr;
+	}
+
+	return SlateFontInfo;
+}
diff --git a/Engine/Source/Runtime/SlateBaseRenderer/Public/FX/SlateFXBaseSubsystem.h b/Engine/Source/Runtime/SlateBaseRenderer/Public/FX/SlateFXBaseSubsystem.h
new file mode 100644
--- /dev/null
+++ b/Engine/Source/Runtime/SlateBaseRenderer/Public/FX/SlateFXBaseSubsystem.h
@@ -0,0 +1,56 @@
+// Copyright Epic Games, Inc. All Rights Reserved.
+
+#pragma once
+
+#include "Subsystems/EngineSubsystem.h"
+#include "Engine/World.h"
+#include "Rendering/SlateRendererTypes.h"
+
+#include "SlateFXBaseSubsystem.generated.h"
+
+#define UE_API SLATEBASERENDERER_API
+
+class USlatePostBufferProcessor;
+
+UCLASS(MinimalAPI, DisplayName = "Slate FX Subsystem")
+class USlateFXBaseSubsystem : public UEngineSubsystem
+{
+	GENERATED_BODY()
+
+public:
+
+	/** Get post processor for a particular post buffer index, if it exists */
+	static UE_API USlatePostBufferProcessor* GetPostProcessor(ESlatePostRT InSlatePostBufferBit);
+
+	//~Begin UEngineSubsystem Interface
+	UE_API virtual bool ShouldCreateSubsystem(UObject* Outer) const override;
+	UE_API virtual void Initialize(FSubsystemCollectionBase& Collection) override;
+	UE_API virtual void Deinitialize() override;
+	//~End UEngineSubsystem Interface
+
+	/** Get post processor for a particular post buffer index, if it exists */
+	UFUNCTION(BlueprintCallable, Category = "SlateFX")
+	UE_API USlatePostBufferProcessor* GetSlatePostProcessor(ESlatePostRT InPostBufferBit);
+
+protected:
+
+	/** Called during world initialization to allow subclasses to populate SlatePostBufferProcessors */
+	virtual void OnInitProcessors() {}
+
+	/** Called during world cleanup to allow subclasses to release renderer-specific resources */
+	virtual void OnCleanupProcessors() {}
+
+	/** Map of post RT buffer index to buffer processors, if they exist */
+	UPROPERTY(Transient)
+	TMap<ESlatePostRT, TObjectPtr<USlatePostBufferProcessor>> SlatePostBufferProcessors;
+
+private:
+
+	/** Callback to create processors on world init */
+	void OnPreWorldInitialization(UWorld* World, const UWorld::InitializationValues IVS);
+
+	/** Callback to remove processors on world cleanup */
+	void OnPostWorldCleanup(UWorld* World, bool SessionEnded, bool bCleanupResources);
+};
+
+#undef UE_API
diff --git a/Engine/Source/Runtime/SlateBaseRenderer/Public/FX/SlatePostBufferBlur.h b/Engine/Source/Runtime/SlateBaseRenderer/Public/FX/SlatePostBufferBlur.h
new file mode 100644
--- /dev/null
+++ b/Engine/Source/Runtime/SlateBaseRenderer/Public/FX/SlatePostBufferBlur.h
@@ -0,0 +1,31 @@
+// Copyright Epic Games, Inc. All Rights Reserved.
+
+#pragma once
+
+#include "FX/SlatePostBufferProcessor.h"
+#include "SlatePostBufferBlur.generated.h"
+
+#define UE_API SLATEBASERENDERER_API
+
+/**
+ * Slate Post Buffer Processor that performs a simple gaussian blur to the backbuffer
+ * 
+ * Create a new asset deriving from this class to use / modify settings.
+ */
+UCLASS(MinimalAPI, Abstract, Blueprintable, CollapseCategories)
+class USlatePostBufferBlur : public USlatePostBufferProcessor
+{
+	GENERATED_BODY()
+
+public:
+
+	UPROPERTY(interp, BlueprintReadWrite, Category = "GaussianBlur")
+	float GaussianBlurStrength = 10;
+
+public:
+
+	USlatePostBufferBlur() {}
+	virtual ~USlatePostBufferBlur() override {}
+};
+
+#undef UE_API
diff --git a/Engine/Source/Runtime/SlateBaseRenderer/Public/FX/SlatePostBufferProcessor.h b/Engine/Source/Runtime/SlateBaseRenderer/Public/FX/SlatePostBufferProcessor.h
new file mode 100644
--- /dev/null
+++ b/Engine/Source/Runtime/SlateBaseRenderer/Public/FX/SlatePostBufferProcessor.h
@@ -0,0 +1,44 @@
+// Copyright Epic Games, Inc. All Rights Reserved.
+
+#pragma once
+
+#include "UObject/Object.h"
+#include "SlatePostBufferProcessor.generated.h"
+
+#define UE_API SLATE_API
+
+/**
+ * Base class for proxy for post buffer processor that the renderthread uses to perform processing.
+ */
+class FSlatePostBufferProcessorProxy : public TSharedFromThis<FSlatePostBufferProcessorProxy>
+{
+public:
+	virtual ~FSlatePostBufferProcessorProxy() {}
+};
+
+/**
+ * Base class for types that can process the backbuffer scene into the slate post buffer.
+ *
+ * You need to register a proxy that derives from 'FSlateRHIPostBufferProcessorProxy' in USlateFXSubsystem::CreatePostBufferProcessorProxy
+ * For an example see: USlatePostBufferBlur.
+ */
+UCLASS(MinimalAPI, Abstract, Blueprintable, CollapseCategories)
+class USlatePostBufferProcessor : public UObject
+{
+	GENERATED_BODY()
+
+public:
+
+	virtual ~USlatePostBufferProcessor() = default;
+
+	void SetRenderThreadProxy(TSharedPtr<FSlatePostBufferProcessorProxy> Proxy)
+	{
+		RenderThreadProxy = Proxy;
+	}
+
+private:
+	TSharedPtr<FSlatePostBufferProcessorProxy> RenderThreadProxy;
+};
+
+
+#undef UE_API
diff --git a/Engine/Source/Runtime/SlateBaseRenderer/Public/SlateFontInfoBlueprintLibrary.h b/Engine/Source/Runtime/SlateBaseRenderer/Public/SlateFontInfoBlueprintLibrary.h
new file mode 100644
--- /dev/null
+++ b/Engine/Source/Runtime/SlateBaseRenderer/Public/SlateFontInfoBlueprintLibrary.h
@@ -0,0 +1,18 @@
+// Copyright Epic Games, Inc. All Rights Reserved.
+
+#pragma once
+
+#include "Fonts/SlateFontInfo.h"
+#include "UObject/Class.h"
+#include "SlateFontInfoBlueprintLibrary.generated.h"
+
+SLATEBASERENDERER_API DECLARE_LOG_CATEGORY_EXTERN(LogSlateFontInfoBlueprint, Log, All);
+
+UCLASS()
+class USlateFontInfoBlueprintLibrary : public UObject
+{
+	GENERATED_BODY()
+
+	UFUNCTION(BlueprintPure, Category = "Slate Font", meta = (NativeMakeFunc))
+	static SLATEBASERENDERER_API FSlateFontInfo MakeSlateFontInfo(const UObject* FontObject, UObject* FontMaterial, FFontOutlineSettings OutlineSettings, FName TypefaceFontName, float Size = 24.0f, int32 LetterSpacing = 0, float SkewAmount = 0.0f, bool bForceMonospaced = false, bool bMaterialIsStencil = false, float MonospacedWidth = 1.0f);
+};
diff --git a/Engine/Source/Runtime/SlateBaseRenderer/SlateBaseRenderer.Build.cs b/Engine/Source/Runtime/SlateBaseRenderer/SlateBaseRenderer.Build.cs
new file mode 100644
--- /dev/null
+++ b/Engine/Source/Runtime/SlateBaseRenderer/SlateBaseRenderer.Build.cs
@@ -0,0 +1,19 @@
+// Copyright Epic Games, Inc. All Rights Reserved.
+
+using UnrealBuildTool;
+
+public class SlateBaseRenderer : ModuleRules
+{
+	public SlateBaseRenderer(ReadOnlyTargetRules Target) : base(Target)
+	{
+		PublicDependencyModuleNames.AddRange(
+			new string[] {
+				"Core",
+				"CoreUObject",
+				"Slate",
+				"SlateCore",
+				"Engine",
+			}
+		);
+	}
+}
diff --git a/Engine/Source/Runtime/SlateCore/Public/Fonts/SlateFontInfo.h b/Engine/Source/Runtime/SlateCore/Public/Fonts/SlateFontInfo.h
--- a/Engine/Source/Runtime/SlateCore/Public/Fonts/SlateFontInfo.h
+++ b/Engine/Source/Runtime/SlateCore/Public/Fonts/SlateFontInfo.h
@@ -142,7 +142,7 @@ struct TStructOpsTypeTraits<FFontOutlineSettings>
 /**
  * A representation of a font in Slate.
  */
-USTRUCT(BlueprintType, meta = (HasNativeMake = "/Script/SlateRHIRenderer.SlateFontBlueprintLibrary.MakeSlateFontInfo"))
+USTRUCT(BlueprintType, meta = (HasNativeMake = "/Script/SlateBaseRenderer.SlateFontInfoBlueprintLibrary.MakeSlateFontInfo"))
 struct FSlateFontInfo
 {
 	GENERATED_USTRUCT_BODY()
diff --git a/Engine/Source/Runtime/SlateNullRenderer/SlateNullRenderer.Build.cs b/Engine/Source/Runtime/SlateNullRenderer/SlateNullRenderer.Build.cs
--- a/Engine/Source/Runtime/SlateNullRenderer/SlateNullRenderer.Build.cs
+++ b/Engine/Source/Runtime/SlateNullRenderer/SlateNullRenderer.Build.cs
@@ -21,7 +21,8 @@ public class SlateNullRenderer : ModuleRules
 				new string[] {
 					"Engine",
 					"RenderCore",
-					"RHI"
+					"RHI",
+					"SlateBaseRenderer"
 				});
 		}
 		else
diff --git a/Engine/Source/Runtime/SlateRHIRenderer/Private/FX/SlateFXSubsystem.cpp b/Engine/Source/Runtime/SlateRHIRenderer/Private/FX/SlateFXSubsystem.cpp
--- a/Engine/Source/Runtime/SlateRHIRenderer/Private/FX/SlateFXSubsystem.cpp
+++ b/Engine/Source/Runtime/SlateRHIRenderer/Private/FX/SlateFXSubsystem.cpp
@@ -3,28 +3,13 @@
 #include "FX/SlateFXSubsystem.h"
 
 #include "Engine/Engine.h"
-#include "Engine/EngineTypes.h"
-#include "Engine/GameInstance.h"
+#include "FX/SlatePostBufferBlur.h"
+#include "FX/SlateRHIPostBufferBlur.h"
 #include "FX/SlateRHIPostBufferProcessor.h"
 #include "SlateRHIRendererSettings.h"
 
 #include UE_INLINE_GENERATED_CPP_BY_NAME(SlateFXSubsystem)
 
-USlateRHIPostBufferProcessor* USlateFXSubsystem::GetPostProcessor(ESlatePostRT InSlatePostBufferBit)
-{
-	USlateRHIPostBufferProcessor* Result = nullptr;
-
-	if (GEngine)
-	{
-		if (USlateFXSubsystem* SlateFXSubsystem = GEngine->GetEngineSubsystem<USlateFXSubsystem>())
-		{
-			Result = SlateFXSubsystem->GetSlatePostProcessor(InSlatePostBufferBit);
-		}
-	}
-
-	return Result;
-}
-
 TSharedPtr<FSlateRHIPostBufferProcessorProxy> USlateFXSubsystem::GetPostProcessorProxy(ESlatePostRT InSlatePostBufferBit)
 {
 	TSharedPtr<FSlateRHIPostBufferProcessorProxy> Result = nullptr;
@@ -48,33 +33,6 @@ void USlateFXSubsystem::BeginDestroy()
 	Super::BeginDestroy();
 }
 
-bool USlateFXSubsystem::ShouldCreateSubsystem(UObject* Outer) const
-{
-	TArray<UClass*> ChildClasses;
-	GetDerivedClasses(GetClass(), ChildClasses, false);
-
-	// Only create an instance if there is no override implementation defined elsewhere
-	return ChildClasses.Num() == 0;
-}
-
-void USlateFXSubsystem::Initialize(FSubsystemCollectionBase& Collection)
-{
-	Super::Initialize(Collection);
-
-	FWorldDelegates::OnPreWorldInitialization.AddUObject(this, &USlateFXSubsystem::OnPreWorldInitialization);
-	FWorldDelegates::OnPostWorldCleanup.AddUObject(this, &USlateFXSubsystem::OnPostWorldCleanup);
-}
-
-void USlateFXSubsystem::Deinitialize()
-{
-	SlatePostBufferProcessors.Empty();
-
-	FWorldDelegates::OnPreWorldInitialization.RemoveAll(this);
-	FWorldDelegates::OnPostWorldCleanup.RemoveAll(this);
-
-	Super::Deinitialize();
-}
-
 TSharedPtr<FSlateRHIPostBufferProcessorProxy> USlateFXSubsystem::GetSlatePostProcessorProxy(ESlatePostRT InPostBufferBit)
 {
 	if (TSharedPtr<FSlateRHIPostBufferProcessorProxy>* ProcessorProxy = SlatePostBufferProcessorProxies.Find(InPostBufferBit))
@@ -85,24 +43,33 @@ TSharedPtr<FSlateRHIPostBufferProcessorProxy> USlateFXSubsystem::GetSlatePostPro
 	return nullptr;
 }
 
-USlateRHIPostBufferProcessor* USlateFXSubsystem::GetSlatePostProcessor(ESlatePostRT InPostBufferBit)
+USlatePostBufferProcessor* USlateFXSubsystem::GetSlatePostProcessor(ESlatePostRT InPostBufferBit)
 {
-	if (TObjectPtr<USlateRHIPostBufferProcessor>* Processor = SlatePostBufferProcessors.Find(InPostBufferBit))
-	{
-		return *Processor;
-	}
-
-	return nullptr;
+	return Super::GetSlatePostProcessor(InPostBufferBit);
 }
 
-void USlateFXSubsystem::OnPreWorldInitialization(UWorld* World, const UWorld::InitializationValues IVS)
+TSharedPtr<FSlateRHIPostBufferProcessorProxy> USlateFXSubsystem::CreatePostBufferProcessorProxy(TSubclassOf<USlatePostBufferProcessor> PostProcessorClass)
 {
-	if (World && World->WorldType == EWorldType::EditorPreview)
+	if (PostProcessorClass)
+	{
+		if (PostProcessorClass->IsChildOf(USlatePostBufferBlur::StaticClass()))
+		{
+			return MakeShared<FSlatePostBufferBlurProxy>();
+		}
+		//Add other processor proxy here.
+
+		UE_LOGF(LogSlate, Warning, "Failed to create FSlateRHIPostBufferProcessorProxy for USlatePostBufferProcessor %ls", *(PostProcessorClass.Get()->GetName()));
+	}
+	else
 	{
-		return;
+		UE_LOGF(LogSlate, Warning, "Failed to create FSlateRHIPostBufferProcessorProxy for an invalid USlatePostBufferProcessor");
 	}
 
-	SlatePostBufferProcessors.Empty();
+	return TSharedPtr<FSlateRHIPostBufferProcessorProxy>();
+}
+
+void USlateFXSubsystem::OnInitProcessors()
+{
 	SlatePostBufferProcessorProxies.Empty();
 
 	if (const USlateRHIRendererSettings* SlateRendererSettings = USlateRHIRendererSettings::Get())
@@ -112,25 +79,24 @@ void USlateFXSubsystem::OnPreWorldInitialization(UWorld* World, const UWorld::In
 			const FSlatePostSettings& PostSetting = SlateRendererSettings->GetSlatePostSetting(SlatePostBufferBit);
 			if (PostSetting.bEnabled && PostSetting.PostProcessorClass)
 			{
-				if (TObjectPtr<USlateRHIPostBufferProcessor> BufferProcessor = NewObject<USlateRHIPostBufferProcessor>(this, PostSetting.PostProcessorClass))
+				if (TObjectPtr<USlatePostBufferProcessor> BufferProcessor = NewObject<USlatePostBufferProcessor>(this, PostSetting.PostProcessorClass))
 				{
-					TSharedPtr<FSlateRHIPostBufferProcessorProxy> PostProcessorProxy = BufferProcessor->GetRenderThreadProxy();
-					SlatePostBufferProcessors.Add(SlatePostBufferBit, BufferProcessor);
-					SlatePostBufferProcessorProxies.Add(SlatePostBufferBit, PostProcessorProxy);
-					PostProcessorProxy->OnUpdateValuesRenderThread();
+					TSharedPtr<FSlateRHIPostBufferProcessorProxy> PostProcessorProxy = CreatePostBufferProcessorProxy(PostSetting.PostProcessorClass);
+					if (ensure(PostProcessorProxy))
+					{
+						BufferProcessor->SetRenderThreadProxy(PostProcessorProxy);
+						PostProcessorProxy->SetOwningProcessorObject(BufferProcessor);
+						SlatePostBufferProcessors.Add(SlatePostBufferBit, BufferProcessor);
+						SlatePostBufferProcessorProxies.Add(SlatePostBufferBit, PostProcessorProxy);
+						PostProcessorProxy->OnUpdateValuesRenderThread();
+					}
 				}
 			}
 		}
 	}
 }
 
-void USlateFXSubsystem::OnPostWorldCleanup(UWorld* World, bool SessionEnded, bool bCleanupResources)
+void USlateFXSubsystem::OnCleanupProcessors()
 {
-	if (World && World->WorldType == EWorldType::EditorPreview)
-	{
-		return;
-	}
-
-	SlatePostBufferProcessors.Empty();
 	SlatePostBufferProcessorProxies.Empty();
 }
diff --git a/Engine/Source/Runtime/SlateRHIRenderer/Private/FX/SlatePostBufferBlur.cpp b/Engine/Source/Runtime/SlateRHIRenderer/Private/FX/SlatePostBufferBlur.cpp
deleted file mode 100644
diff --git a/Engine/Source/Runtime/SlateRHIRenderer/Private/FX/SlateRHIPostBufferBlur.cpp b/Engine/Source/Runtime/SlateRHIRenderer/Private/FX/SlateRHIPostBufferBlur.cpp
new file mode 100644
--- /dev/null
+++ b/Engine/Source/Runtime/SlateRHIRenderer/Private/FX/SlateRHIPostBufferBlur.cpp
@@ -0,0 +1,76 @@
+// Copyright Epic Games, Inc. All Rights Reserved.
+
+#include "FX/SlateRHIPostBufferBlur.h"
+#include "FX/SlatePostBufferBlur.h"
+#include "UObject/StrongObjectPtr.h"
+#include "UObject/WeakObjectPtrTemplates.h"
+#include "RHIResources.h"
+#include "RHIUtilities.h"
+
+/////////////////////////////////////////////////////
+// FSlatePostBufferBlurProxy
+
+void FSlatePostBufferBlurProxy::PostProcess_Renderthread(FRDGBuilder& GraphBuilder, const FScreenPassTexture& InputTexture, const FScreenPassTexture& OutputTexture)
+{
+	RDG_EVENT_SCOPE(GraphBuilder, "SlatePostBufferBlur");
+
+	if (GaussianBlurStrength_RenderThread < UE_SMALL_NUMBER)
+	{
+		AddSlatePostProcessCopy(GraphBuilder, InputTexture, OutputTexture);
+	}
+	else
+	{
+		FSlatePostProcessSimpleBlurPassInputs BlurInputs;
+		BlurInputs.InputTexture = InputTexture;
+		BlurInputs.OutputTexture = OutputTexture;
+		BlurInputs.Strength = GaussianBlurStrength_RenderThread;
+
+		AddSlatePostProcessBlurPass(GraphBuilder, BlurInputs);
+	}
+}
+
+void FSlatePostBufferBlurProxy::OnUpdateValuesRenderThread()
+{
+	// Don't issue multiple updates in a single frame from the CPU based on dirty values
+	if (!ParamUpdateFence.IsFenceComplete())
+	{
+		return;
+	}
+
+	// Only issue an update when parent exists & values are different
+	if (USlatePostBufferBlur* ParentBlurObject = Cast<USlatePostBufferBlur>(ParentObject))
+	{
+		if (ParentBlurObject->GaussianBlurStrength != GaussianBlurStrength_RenderThread)
+		{
+			// Blur strengths can be updated from renderthread during draw or gamethread,
+			// if our parent object value matches the predraw then don't update the renderthread value.
+			// Instead we need to update our parent object's value to match the last value from renderthread
+			bool bUpdatedInRenderThread = ParentBlurObject->GaussianBlurStrength == GaussianBlurStrengthPreDraw;
+
+			if (bUpdatedInRenderThread)
+			{
+				ParentBlurObject->GaussianBlurStrength = GaussianBlurStrength_RenderThread;
+				GaussianBlurStrengthPreDraw = GaussianBlurStrength_RenderThread;
+			}
+			else
+			{
+				// Explicit param copy to avoid renderthread from reading value during gamethread write
+				float GaussianBlurStrengthCopy = ParentBlurObject->GaussianBlurStrength;
+				GaussianBlurStrengthPreDraw = GaussianBlurStrengthCopy;
+
+				// Execute param copy in a render command to safely update value on renderthread without race conditions
+				TWeakPtr<FSlatePostBufferBlurProxy> TempWeakThis = SharedThis(this);
+				ENQUEUE_RENDER_COMMAND(FUpdateValuesRenderThreadFX_Blur)([TempWeakThis, GaussianBlurStrengthCopy](FRHICommandListImmediate& RHICmdList)
+				{
+					if (TSharedPtr<FSlatePostBufferBlurProxy> SharedThisPin = TempWeakThis.Pin())
+					{
+						SharedThisPin->GaussianBlurStrength_RenderThread = GaussianBlurStrengthCopy;
+					}
+				});
+
+				// Issue fence to prevent multiple updates in a single frame
+				ParamUpdateFence.BeginFence();
+			}
+		}
+	}
+}
\ No newline at end of file
diff --git a/Engine/Source/Runtime/SlateRHIRenderer/Private/FX/SlateRHIPostBufferProcessor.cpp b/Engine/Source/Runtime/SlateRHIRenderer/Private/FX/SlateRHIPostBufferProcessor.cpp
deleted file mode 100644
diff --git a/Engine/Source/Runtime/SlateRHIRenderer/Private/SlateFontBlueprintLibrary.h b/Engine/Source/Runtime/SlateRHIRenderer/Private/SlateFontBlueprintLibrary.h
--- a/Engine/Source/Runtime/SlateRHIRenderer/Private/SlateFontBlueprintLibrary.h
+++ b/Engine/Source/Runtime/SlateRHIRenderer/Private/SlateFontBlueprintLibrary.h
@@ -13,6 +13,6 @@ class USlateFontBlueprintLibrary : public UObject
 {
 	GENERATED_BODY()
 
-	UFUNCTION(BlueprintPure, Category = "Slate Font", meta = (NativeMakeFunc))
+	UFUNCTION(BlueprintPure, Category = "Slate Font", meta = (DeprecatedFunction, DeprecationMessage = "MakeSlateFontInfo from SlateFontBlueprintLibrary has been deprecated, use MakeSlateFontInfo from SlateFontInfoBlueprintLibrary instead."))
 	static SLATERHIRENDERER_API FSlateFontInfo MakeSlateFontInfo(const UObject* FontObject, UObject* FontMaterial, FFontOutlineSettings OutlineSettings, FName TypefaceFontName, float Size = 24.0f, int32 LetterSpacing = 0, float SkewAmount = 0.0f, bool bForceMonospaced = false, bool bMaterialIsStencil = false, float MonospacedWidth = 1.0f);
 };
\ No newline at end of file
diff --git a/Engine/Source/Runtime/SlateRHIRenderer/Private/SlateRHIRenderer.h b/Engine/Source/Runtime/SlateRHIRenderer/Private/SlateRHIRenderer.h
--- a/Engine/Source/Runtime/SlateRHIRenderer/Private/SlateRHIRenderer.h
+++ b/Engine/Source/Runtime/SlateRHIRenderer/Private/SlateRHIRenderer.h
@@ -18,14 +18,8 @@
 
 class FSlateElementBatcher;
 class FSlateRHIRenderingPolicy;
-class FSlateRHIRenderingPolicyInterface;
-class FSlateRHIPostBufferProcessorProxy;
-class USlateRHIPostBufferProcessor;
-class USlateRHIRendererSettings;
 class ISlateStyle;
 class SWindow;
-struct FSlatePostSettings;
-struct FSlateDrawWindowsCommand;
 
 extern FMatrix CreateSlateProjectionMatrix(uint32 Width, uint32 Height);
 
diff --git a/Engine/Source/Runtime/SlateRHIRenderer/Public/FX/SlateFXSubsystem.h b/Engine/Source/Runtime/SlateRHIRenderer/Public/FX/SlateFXSubsystem.h
--- a/Engine/Source/Runtime/SlateRHIRenderer/Public/FX/SlateFXSubsystem.h
+++ b/Engine/Source/Runtime/SlateRHIRenderer/Public/FX/SlateFXSubsystem.h
@@ -2,61 +2,42 @@
 
 #pragma once
 
-#include "Subsystems/EngineSubsystem.h"
-
-#include "Engine/World.h"
+#include "FX/SlateFXBaseSubsystem.h"
 
 #include "SlateFXSubsystem.generated.h"
 
 #define UE_API SLATERHIRENDERER_API
 
 class FSlateRHIPostBufferProcessorProxy;
-class USlateRHIPostBufferProcessor;
+class USlatePostBufferProcessor;
 
-UCLASS(MinimalAPI, DisplayName = "Slate FX Subsystem")
-class USlateFXSubsystem : public UEngineSubsystem
+UCLASS(MinimalAPI, DisplayName = "Slate FX Subsystem (RHI)")
+class USlateFXSubsystem : public USlateFXBaseSubsystem
 {
 	GENERATED_BODY()
 
 public:
 
-	static UE_API USlateRHIPostBufferProcessor* GetPostProcessor(ESlatePostRT InSlatePostBufferBit);
 	static UE_API TSharedPtr<FSlateRHIPostBufferProcessorProxy> GetPostProcessorProxy(ESlatePostRT InSlatePostBufferBit);
 
 	//~ Begin UObject Interface.
 	UE_API virtual void BeginDestroy() override;
 	//~ End UObject Interface.
 
-	//~Begin UGameInstanceSubsystem Interface
-	UE_API virtual bool ShouldCreateSubsystem(UObject* Outer) const override;
-	UE_API virtual void Initialize(FSubsystemCollectionBase& Collection) override;
-	UE_API virtual void Deinitialize() override;
-	//~End UGameInstanceSubsystem Interface
-
 	/** Get post processor proxy for a particular post buffer index, if it exists */
 	UE_API TSharedPtr<FSlateRHIPostBufferProcessorProxy> GetSlatePostProcessorProxy(ESlatePostRT InPostBufferBit);
 
-public:
-
 	/** Get post processor for a particular post buffer index, if it exists */
-	UFUNCTION(BlueprintCallable, Category = "SlateFX")
-	UE_API USlateRHIPostBufferProcessor* GetSlatePostProcessor(ESlatePostRT InPostBufferBit);
+	UE_API USlatePostBufferProcessor* GetSlatePostProcessor(ESlatePostRT InPostBufferBit);
 
-private:
+protected:
 
-	/** Map of post RT buffer index to buffer processors, if they exist */
-	UPROPERTY(Transient)
-	TMap<ESlatePostRT, TObjectPtr<USlateRHIPostBufferProcessor>> SlatePostBufferProcessors;
+	virtual void OnInitProcessors() override;
+	virtual void OnCleanupProcessors() override;
 
 private:
 
-	/** Callback to create processors on world init */
-	void OnPreWorldInitialization(UWorld* World, const UWorld::InitializationValues IVS);
-
-	/** Callback to remove processors on world cleanup */
-	void OnPostWorldCleanup(UWorld* World, bool SessionEnded, bool bCleanupResources);
-
-private:
+	static TSharedPtr<FSlateRHIPostBufferProcessorProxy> CreatePostBufferProcessorProxy(TSubclassOf<USlatePostBufferProcessor> PostProcessorClass);
 
 	/** Map of post RT buffer index to buffer processor renderthread proxies, if they exist */
 	TMap<ESlatePostRT, TSharedPtr<FSlateRHIPostBufferProcessorProxy>> SlatePostBufferProcessorProxies;
diff --git a/Engine/Source/Runtime/SlateRHIRenderer/Public/FX/SlatePostBufferBlur.h b/Engine/Source/Runtime/SlateRHIRenderer/Public/FX/SlatePostBufferBlur.h
deleted file mode 100644
diff --git a/Engine/Source/Runtime/SlateRHIRenderer/Public/FX/SlateRHIPostBufferBlur.h b/Engine/Source/Runtime/SlateRHIRenderer/Public/FX/SlateRHIPostBufferBlur.h
new file mode 100644
--- /dev/null
+++ b/Engine/Source/Runtime/SlateRHIRenderer/Public/FX/SlateRHIPostBufferBlur.h
@@ -0,0 +1,40 @@
+// Copyright Epic Games, Inc. All Rights Reserved.
+
+#pragma once
+
+#include "FX/SlateRHIPostBufferProcessor.h"
+
+#define UE_API SLATERHIRENDERER_API
+
+/**
+ * Proxy for post buffer processor that the renderthread uses to perform processing
+ * This proxy exists because generally speaking usage on UObjects on the renderthread
+ * is a race condition due to UObjects being managed / updated by the game thread
+ */
+class FSlatePostBufferBlurProxy : public FSlateRHIPostBufferProcessorProxy
+{
+
+public:
+
+	//~ Begin FSlateRHIPostBufferProcessorProxy Interface
+	UE_API virtual void PostProcess_Renderthread(FRDGBuilder& GraphBuilder, const FScreenPassTexture& InputTexture, const FScreenPassTexture& OutputTexture) override;
+	UE_API virtual void OnUpdateValuesRenderThread() override;
+	//~ End FSlateRHIPostBufferProcessorProxy Interface
+
+	/** Blur strength to use when processing, renderthread version actually used to draw. Must be updated via render command except during initialization. */
+	float GaussianBlurStrength_RenderThread = 10;
+
+	/** 
+	 * Blur strength can be updated from both renderthread during draw and gamethread update. 
+	 * Store the last value gamethread provided so we know if we should use the renderthread value or gamethread value. 
+	 * We will use the most recently updated one.
+	 */
+	float GaussianBlurStrengthPreDraw = 10;
+
+protected:
+
+	/** Fence to allow for us to queue only one update per draw command from the gamethread */
+	FRenderCommandFence ParamUpdateFence;
+};
+
+#undef UE_API
diff --git a/Engine/Source/Runtime/SlateRHIRenderer/Public/FX/SlateRHIPostBufferProcessor.h b/Engine/Source/Runtime/SlateRHIRenderer/Public/FX/SlateRHIPostBufferProcessor.h
--- a/Engine/Source/Runtime/SlateRHIRenderer/Public/FX/SlateRHIPostBufferProcessor.h
+++ b/Engine/Source/Runtime/SlateRHIRenderer/Public/FX/SlateRHIPostBufferProcessor.h
@@ -3,6 +3,7 @@
 #pragma once
 
 #include "Engine/TextureRenderTarget2D.h"
+#include "FX/SlatePostBufferProcessor.h"
 #include "Layout/PaintGeometry.h"
 #include "Rendering/RenderingCommon.h"
 #include "Rendering/SlateRenderer.h"
@@ -11,8 +12,6 @@
 #include "Interfaces/SlateRHIRenderingPolicyInterface.h"
 #include "ScreenPass.h"
 
-#include "SlateRHIPostBufferProcessor.generated.h"
-
 class USlateRHIPostBufferProcessor;
 
 /**
@@ -20,7 +19,7 @@ class USlateRHIPostBufferProcessor;
  * This proxy exists because generally speaking usage on UObjects on the renderthread
  * is a race condition due to UObjects being managed / updated by the game thread
  */
-class FSlateRHIPostBufferProcessorProxy : public TSharedFromThis<FSlateRHIPostBufferProcessorProxy>
+class FSlateRHIPostBufferProcessorProxy : public FSlatePostBufferProcessorProxy
 {
 
 public:
@@ -32,13 +31,6 @@ public:
 	{
 	}
 
-	UE_DEPRECATED(5.5, "Use RDG version instead.")
-PRAGMA_DISABLE_DEPRECATION_WARNINGS
-	virtual void PostProcess_Renderthread(FRHICommandListImmediate& RHICmdList, FRHITexture* Src, FRHITexture* Dst, FIntRect SrcRect, FIntRect DstRect, FSlateRHIRenderingPolicyInterface InRenderingPolicy)
-	{
-	}
-PRAGMA_ENABLE_DEPRECATION_WARNINGS
-
 	/** 
 	 * Called when an post buffer update element is added to a renderbatch, 
 	 * gives proxies a chance to queue updates to their renderthread values based on the UObject processor.
@@ -51,7 +43,7 @@ PRAGMA_ENABLE_DEPRECATION_WARNINGS
 	/**
 	 * Set the UObject that we are a renderthread proxy for, useful for doing gamethread updates from the proxy
 	 */
-	void SetOwningProcessorObject(USlateRHIPostBufferProcessor* InParentObject)
+	void SetOwningProcessorObject(USlatePostBufferProcessor* InParentObject)
 	{
 		ParentObject = InParentObject;
 	}
@@ -59,43 +51,22 @@ PRAGMA_ENABLE_DEPRECATION_WARNINGS
 protected:
 
 	/** Pointer to processor that we are a proxy for, external design constraints should ensure that this is always valid */
-	TWeakObjectPtr<USlateRHIPostBufferProcessor> ParentObject;
+	TWeakObjectPtr<USlatePostBufferProcessor> ParentObject;
 };
 
 /**
  * Base class for types that can process the backbuffer scene into the slate post buffer.
  * 
- * Implement 'PostProcess' in your derived class. Additionally, you need to create a renderthread proxy that derives from 'FSlateRHIPostBufferProcessorProxy'
+ * You need to create a renderthread proxy that derives from 'FSlateRHIPostBufferProcessorProxy'
  * For an example see: USlatePostBufferBlur.
  */
-UCLASS(MinimalAPI, Abstract, Blueprintable, CollapseCategories)
-class USlateRHIPostBufferProcessor : public UObject
+class UE_DEPRECATED(5.8, "Use USlatePostBufferProcessor instead") USlateRHIPostBufferProcessor : public UObject
 {
-	GENERATED_BODY()
 
 public:
 
 	virtual ~USlateRHIPostBufferProcessor() = default;
 
-	/**
-	 * Overridable postprocess for the given source scene backbuffer provided in 'Src' into 'Dst'
-	 * You must override this method. In your override, you should copy params before executing 'ENQUEUE_RENDER_COMMAND'.
-	 * This allows you to avoid render & game thread race conditions. See 'SlatePostBufferBlur' for example.
-	 * Also avoid capturing [this] in your override, to avoid possible GC issues with the processor instance.
-	 *
-	 * @param InViewInfo			'FViewportInfo' resource used to get backbuffer in standalone
-	 * @param InViewportTexture		'FSlateRenderTargetRHI' resource used to get the 'BufferedRT' viewport texture used in PIE
-	 * @param InElementWindowSize	Size of window being rendered, used to determine if using stereo rendering or not.
-	 * @param InRenderingPolicy		Slate RHI RenderingPolicy
-	 * @param InSlatePostBuffer		Texture render target used for final output
-	 */
-	UE_DEPRECATED(5.5, "This path is longer supported. Get the render proxy instead.")
-PRAGMA_DISABLE_DEPRECATION_WARNINGS
-	virtual void PostProcess(FRenderResource* InViewInfo, FRenderResource* InViewportTexture, FVector2D InElementWindowSize, FSlateRHIRenderingPolicyInterface InRenderingPolicy, UTextureRenderTarget2D* InSlatePostBuffer)
-	{
-	}
-PRAGMA_ENABLE_DEPRECATION_WARNINGS
-
 	/**
 	 * Gets proxy for this post buffer processor, for execution on the renderthread
 	 */
diff --git a/Engine/Source/Runtime/SlateRHIRenderer/Public/SlateRHIRendererSettings.h b/Engine/Source/Runtime/SlateRHIRenderer/Public/SlateRHIRendererSettings.h
--- a/Engine/Source/Runtime/SlateRHIRenderer/Public/SlateRHIRendererSettings.h
+++ b/Engine/Source/Runtime/SlateRHIRenderer/Public/SlateRHIRendererSettings.h
@@ -15,7 +15,7 @@
 
 #define UE_API SLATERHIRENDERER_API
 
-class USlateRHIPostBufferProcessor;
+class USlatePostBufferProcessor;
 class UTextureRenderTarget2D;
 
 /**
@@ -58,7 +58,7 @@ public:
 
 	/** Copy of actually loaded post processor class */
 	UPROPERTY(EditAnywhere, BlueprintReadWrite, Category = BufferSettings, Meta = (AllowAbstract = false))
-	TSubclassOf<USlateRHIPostBufferProcessor> PostProcessorClass;
+	TSubclassOf<USlatePostBufferProcessor> PostProcessorClass;
 
 private:
 
diff --git a/Engine/Source/Runtime/SlateRHIRenderer/SlateRHIRenderer.Build.cs b/Engine/Source/Runtime/SlateRHIRenderer/SlateRHIRenderer.Build.cs
--- a/Engine/Source/Runtime/SlateRHIRenderer/SlateRHIRenderer.Build.cs
+++ b/Engine/Source/Runtime/SlateRHIRenderer/SlateRHIRenderer.Build.cs
@@ -6,6 +6,12 @@ public class SlateRHIRenderer : ModuleRules
 {
     public SlateRHIRenderer(ReadOnlyTargetRules Target) : base(Target)
 	{
+		PublicDependencyModuleNames.AddRange(
+			new string[] {
+				"SlateBaseRenderer"
+			}
+		);
+
 		PrivateDependencyModuleNames.AddRange(
 			new string[] {
 				"Core",
diff --git a/Engine/Source/Runtime/UMG/Private/Components/PostBufferBlurUpdater.cpp b/Engine/Source/Runtime/UMG/Private/Components/PostBufferBlurUpdater.cpp
--- a/Engine/Source/Runtime/UMG/Private/Components/PostBufferBlurUpdater.cpp
+++ b/Engine/Source/Runtime/UMG/Private/Components/PostBufferBlurUpdater.cpp
@@ -5,7 +5,7 @@
 // Exclude SlateRHIRenderer related includes & implementations from server builds since the module is not a dependency / will not link for UMG on the server
 #if !UE_SERVER
 #include "FX/SlateRHIPostBufferProcessor.h"
-#include "FX/SlatePostBufferBlur.h"
+#include "FX/SlateRHIPostBufferBlur.h"
 #endif // !UE_SERVER
 
 #include UE_INLINE_GENERATED_CPP_BY_NAME(PostBufferBlurUpdater)
