[Backout] - CL52313506
#rnx
Original CL Desc
-----------------------------------------------------------------
When running with a null renderer, SlateRHIRenderer is not loaded, leading to errors in BP, as some required types can't be found.
=> extracted blueprintable content from SlateRHIRenderer into (the new) SlateBaseRenderer

#tests Engine test now have an automated test using post buffer blur. Ran with null renderer and validated that there was no error with BP (they are still executed properly). Validated that all moved BP functions/classes are deprecated properly by creating first data using it in engine test, then reopening it with new code.

diff --git a/Engine/Config/BaseEngine.ini b/Engine/Config/BaseEngine.ini
--- a/Engine/Config/BaseEngine.ini
+++ b/Engine/Config/BaseEngine.ini
@@ -1413,11 +1413,6 @@ ClassName=/Script/OnlineSubsystemUtils.OnlineEngineInterfaceImpl
 +StructRedirects=(OldName="/Script/Water.ShallowWaterSimulationGrid",NewName="/Script/Water.ShallowWaterSimulationGrid_DEPRECATED")
 +PropertyRedirects=(OldName="/Script/Water.BakedShallowWaterSimulationComponent.SimulationData",NewName="/Script/Water.BakedShallowWaterSimulationComponent.SimulationData_DEPRECATED")
 
-+ClassRedirects=(OldName="/Script/SlateRHIRenderer.SlateRHIPostBufferProcessor",NewName="/Script/SlateBaseRenderer.SlatePostBufferProcessor")
-+ClassRedirects=(OldName="/Script/SlateRHIRenderer.SlatePostBufferBlur",NewName="/Script/SlateBaseRenderer.SlatePostBufferBlur")
-+ClassRedirects=(OldName="/Script/SlateRHIRenderer.SlateFXSubsystem",NewName="/Script/SlateBaseRenderer.SlateFXBaseSubsystem",bInstanceOnly=false)
-+FunctionRedirects=(OldName="/Script/SlateRHIRenderer.SlateFXSubsystem.GetSlatePostProcessor",NewName="/Script/SlateBaseRenderer.SlateFXBaseSubsystem.GetSlatePostProcessor")
-
 ; End of releases redirects
 
 ; Fixing the output of tags for blueprints to be reset.
diff --git a/Engine/Source/Runtime/SlateBaseRenderer/Private/FX/SlateFXBaseSubsystem.cpp b/Engine/Source/Runtime/SlateBaseRenderer/Private/FX/SlateFXBaseSubsystem.cpp
deleted file mode 100644
diff --git a/Engine/Source/Runtime/SlateBaseRenderer/Private/FX/SlatePostBufferBlur.cpp b/Engine/Source/Runtime/SlateBaseRenderer/Private/FX/SlatePostBufferBlur.cpp
deleted file mode 100644
diff --git a/Engine/Source/Runtime/SlateBaseRenderer/Private/FX/SlatePostBufferProcessor.cpp b/Engine/Source/Runtime/SlateBaseRenderer/Private/FX/SlatePostBufferProcessor.cpp
deleted file mode 100644
diff --git a/Engine/Source/Runtime/SlateBaseRenderer/Private/SlateBaseRendererModule.cpp b/Engine/Source/Runtime/SlateBaseRenderer/Private/SlateBaseRendererModule.cpp
deleted file mode 100644
diff --git a/Engine/Source/Runtime/SlateBaseRenderer/Private/SlateFontInfoBlueprintLibrary.cpp b/Engine/Source/Runtime/SlateBaseRenderer/Private/SlateFontInfoBlueprintLibrary.cpp
deleted file mode 100644
diff --git a/Engine/Source/Runtime/SlateBaseRenderer/Public/FX/SlateFXBaseSubsystem.h b/Engine/Source/Runtime/SlateBaseRenderer/Public/FX/SlateFXBaseSubsystem.h
deleted file mode 100644
diff --git a/Engine/Source/Runtime/SlateBaseRenderer/Public/FX/SlatePostBufferBlur.h b/Engine/Source/Runtime/SlateBaseRenderer/Public/FX/SlatePostBufferBlur.h
deleted file mode 100644
diff --git a/Engine/Source/Runtime/SlateBaseRenderer/Public/FX/SlatePostBufferProcessor.h b/Engine/Source/Runtime/SlateBaseRenderer/Public/FX/SlatePostBufferProcessor.h
deleted file mode 100644
diff --git a/Engine/Source/Runtime/SlateBaseRenderer/Public/SlateFontInfoBlueprintLibrary.h b/Engine/Source/Runtime/SlateBaseRenderer/Public/SlateFontInfoBlueprintLibrary.h
deleted file mode 100644
diff --git a/Engine/Source/Runtime/SlateBaseRenderer/SlateBaseRenderer.Build.cs b/Engine/Source/Runtime/SlateBaseRenderer/SlateBaseRenderer.Build.cs
deleted file mode 100644
diff --git a/Engine/Source/Runtime/SlateCore/Public/Fonts/SlateFontInfo.h b/Engine/Source/Runtime/SlateCore/Public/Fonts/SlateFontInfo.h
--- a/Engine/Source/Runtime/SlateCore/Public/Fonts/SlateFontInfo.h
+++ b/Engine/Source/Runtime/SlateCore/Public/Fonts/SlateFontInfo.h
@@ -142,7 +142,7 @@ struct TStructOpsTypeTraits<FFontOutlineSettings>
 /**
  * A representation of a font in Slate.
  */
-USTRUCT(BlueprintType, meta = (HasNativeMake = "/Script/SlateBaseRenderer.SlateFontInfoBlueprintLibrary.MakeSlateFontInfo"))
+USTRUCT(BlueprintType, meta = (HasNativeMake = "/Script/SlateRHIRenderer.SlateFontBlueprintLibrary.MakeSlateFontInfo"))
 struct FSlateFontInfo
 {
 	GENERATED_USTRUCT_BODY()
diff --git a/Engine/Source/Runtime/SlateNullRenderer/SlateNullRenderer.Build.cs b/Engine/Source/Runtime/SlateNullRenderer/SlateNullRenderer.Build.cs
--- a/Engine/Source/Runtime/SlateNullRenderer/SlateNullRenderer.Build.cs
+++ b/Engine/Source/Runtime/SlateNullRenderer/SlateNullRenderer.Build.cs
@@ -21,8 +21,7 @@ public class SlateNullRenderer : ModuleRules
 				new string[] {
 					"Engine",
 					"RenderCore",
-					"RHI",
-					"SlateBaseRenderer"
+					"RHI"
 				});
 		}
 		else
diff --git a/Engine/Source/Runtime/SlateRHIRenderer/Private/FX/SlateFXSubsystem.cpp b/Engine/Source/Runtime/SlateRHIRenderer/Private/FX/SlateFXSubsystem.cpp
--- a/Engine/Source/Runtime/SlateRHIRenderer/Private/FX/SlateFXSubsystem.cpp
+++ b/Engine/Source/Runtime/SlateRHIRenderer/Private/FX/SlateFXSubsystem.cpp
@@ -3,13 +3,28 @@
 #include "FX/SlateFXSubsystem.h"
 
 #include "Engine/Engine.h"
-#include "FX/SlatePostBufferBlur.h"
-#include "FX/SlateRHIPostBufferBlur.h"
+#include "Engine/EngineTypes.h"
+#include "Engine/GameInstance.h"
 #include "FX/SlateRHIPostBufferProcessor.h"
 #include "SlateRHIRendererSettings.h"
 
 #include UE_INLINE_GENERATED_CPP_BY_NAME(SlateFXSubsystem)
 
+USlateRHIPostBufferProcessor* USlateFXSubsystem::GetPostProcessor(ESlatePostRT InSlatePostBufferBit)
+{
+	USlateRHIPostBufferProcessor* Result = nullptr;
+
+	if (GEngine)
+	{
+		if (USlateFXSubsystem* SlateFXSubsystem = GEngine->GetEngineSubsystem<USlateFXSubsystem>())
+		{
+			Result = SlateFXSubsystem->GetSlatePostProcessor(InSlatePostBufferBit);
+		}
+	}
+
+	return Result;
+}
+
 TSharedPtr<FSlateRHIPostBufferProcessorProxy> USlateFXSubsystem::GetPostProcessorProxy(ESlatePostRT InSlatePostBufferBit)
 {
 	TSharedPtr<FSlateRHIPostBufferProcessorProxy> Result = nullptr;
@@ -33,6 +48,33 @@ void USlateFXSubsystem::BeginDestroy()
 	Super::BeginDestroy();
 }
 
+bool USlateFXSubsystem::ShouldCreateSubsystem(UObject* Outer) const
+{
+	TArray<UClass*> ChildClasses;
+	GetDerivedClasses(GetClass(), ChildClasses, false);
+
+	// Only create an instance if there is no override implementation defined elsewhere
+	return ChildClasses.Num() == 0;
+}
+
+void USlateFXSubsystem::Initialize(FSubsystemCollectionBase& Collection)
+{
+	Super::Initialize(Collection);
+
+	FWorldDelegates::OnPreWorldInitialization.AddUObject(this, &USlateFXSubsystem::OnPreWorldInitialization);
+	FWorldDelegates::OnPostWorldCleanup.AddUObject(this, &USlateFXSubsystem::OnPostWorldCleanup);
+}
+
+void USlateFXSubsystem::Deinitialize()
+{
+	SlatePostBufferProcessors.Empty();
+
+	FWorldDelegates::OnPreWorldInitialization.RemoveAll(this);
+	FWorldDelegates::OnPostWorldCleanup.RemoveAll(this);
+
+	Super::Deinitialize();
+}
+
 TSharedPtr<FSlateRHIPostBufferProcessorProxy> USlateFXSubsystem::GetSlatePostProcessorProxy(ESlatePostRT InPostBufferBit)
 {
 	if (TSharedPtr<FSlateRHIPostBufferProcessorProxy>* ProcessorProxy = SlatePostBufferProcessorProxies.Find(InPostBufferBit))
@@ -43,33 +85,24 @@ TSharedPtr<FSlateRHIPostBufferProcessorProxy> USlateFXSubsystem::GetSlatePostPro
 	return nullptr;
 }
 
-USlatePostBufferProcessor* USlateFXSubsystem::GetSlatePostProcessor(ESlatePostRT InPostBufferBit)
+USlateRHIPostBufferProcessor* USlateFXSubsystem::GetSlatePostProcessor(ESlatePostRT InPostBufferBit)
 {
-	return Super::GetSlatePostProcessor(InPostBufferBit);
-}
-
-TSharedPtr<FSlateRHIPostBufferProcessorProxy> USlateFXSubsystem::CreatePostBufferProcessorProxy(TSubclassOf<USlatePostBufferProcessor> PostProcessorClass)
-{
-	if (PostProcessorClass)
-	{
-		if (PostProcessorClass->IsChildOf(USlatePostBufferBlur::StaticClass()))
-		{
-			return MakeShared<FSlatePostBufferBlurProxy>();
-		}
-		//Add other processor proxy here.
-
-		UE_LOGF(LogSlate, Warning, "Failed to create FSlateRHIPostBufferProcessorProxy for USlatePostBufferProcessor %ls", *(PostProcessorClass.Get()->GetName()));
-	}
-	else
+	if (TObjectPtr<USlateRHIPostBufferProcessor>* Processor = SlatePostBufferProcessors.Find(InPostBufferBit))
 	{
-		UE_LOGF(LogSlate, Warning, "Failed to create FSlateRHIPostBufferProcessorProxy for an invalid USlatePostBufferProcessor");
+		return *Processor;
 	}
 
-	return TSharedPtr<FSlateRHIPostBufferProcessorProxy>();
+	return nullptr;
 }
 
-void USlateFXSubsystem::OnInitProcessors()
+void USlateFXSubsystem::OnPreWorldInitialization(UWorld* World, const UWorld::InitializationValues IVS)
 {
+	if (World && World->WorldType == EWorldType::EditorPreview)
+	{
+		return;
+	}
+
+	SlatePostBufferProcessors.Empty();
 	SlatePostBufferProcessorProxies.Empty();
 
 	if (const USlateRHIRendererSettings* SlateRendererSettings = USlateRHIRendererSettings::Get())
@@ -79,24 +112,25 @@ void USlateFXSubsystem::OnInitProcessors()
 			const FSlatePostSettings& PostSetting = SlateRendererSettings->GetSlatePostSetting(SlatePostBufferBit);
 			if (PostSetting.bEnabled && PostSetting.PostProcessorClass)
 			{
-				if (TObjectPtr<USlatePostBufferProcessor> BufferProcessor = NewObject<USlatePostBufferProcessor>(this, PostSetting.PostProcessorClass))
+				if (TObjectPtr<USlateRHIPostBufferProcessor> BufferProcessor = NewObject<USlateRHIPostBufferProcessor>(this, PostSetting.PostProcessorClass))
 				{
-					TSharedPtr<FSlateRHIPostBufferProcessorProxy> PostProcessorProxy = CreatePostBufferProcessorProxy(PostSetting.PostProcessorClass);
-					if (ensure(PostProcessorProxy))
-					{
-						BufferProcessor->SetRenderThreadProxy(PostProcessorProxy);
-						PostProcessorProxy->SetOwningProcessorObject(BufferProcessor);
-						SlatePostBufferProcessors.Add(SlatePostBufferBit, BufferProcessor);
-						SlatePostBufferProcessorProxies.Add(SlatePostBufferBit, PostProcessorProxy);
-						PostProcessorProxy->OnUpdateValuesRenderThread();
-					}
+					TSharedPtr<FSlateRHIPostBufferProcessorProxy> PostProcessorProxy = BufferProcessor->GetRenderThreadProxy();
+					SlatePostBufferProcessors.Add(SlatePostBufferBit, BufferProcessor);
+					SlatePostBufferProcessorProxies.Add(SlatePostBufferBit, PostProcessorProxy);
+					PostProcessorProxy->OnUpdateValuesRenderThread();
 				}
 			}
 		}
 	}
 }
 
-void USlateFXSubsystem::OnCleanupProcessors()
+void USlateFXSubsystem::OnPostWorldCleanup(UWorld* World, bool SessionEnded, bool bCleanupResources)
 {
+	if (World && World->WorldType == EWorldType::EditorPreview)
+	{
+		return;
+	}
+
+	SlatePostBufferProcessors.Empty();
 	SlatePostBufferProcessorProxies.Empty();
 }
diff --git a/Engine/Source/Runtime/SlateRHIRenderer/Private/FX/SlatePostBufferBlur.cpp b/Engine/Source/Runtime/SlateRHIRenderer/Private/FX/SlatePostBufferBlur.cpp
new file mode 100644
--- /dev/null
+++ b/Engine/Source/Runtime/SlateRHIRenderer/Private/FX/SlatePostBufferBlur.cpp
@@ -0,0 +1,101 @@
+// Copyright Epic Games, Inc. All Rights Reserved.
+
+#include "FX/SlatePostBufferBlur.h"
+#include "UObject/StrongObjectPtr.h"
+#include "UObject/WeakObjectPtrTemplates.h"
+#include "RHIResources.h"
+#include "RHIUtilities.h"
+
+#include UE_INLINE_GENERATED_CPP_BY_NAME(SlatePostBufferBlur)
+
+/////////////////////////////////////////////////////
+// FSlatePostBufferBlurProxy
+
+void FSlatePostBufferBlurProxy::PostProcess_Renderthread(FRDGBuilder& GraphBuilder, const FScreenPassTexture& InputTexture, const FScreenPassTexture& OutputTexture)
+{
+	RDG_EVENT_SCOPE(GraphBuilder, "SlatePostBufferBlur");
+
+	if (GaussianBlurStrength_RenderThread < UE_SMALL_NUMBER)
+	{
+		AddSlatePostProcessCopy(GraphBuilder, InputTexture, OutputTexture);
+	}
+	else
+	{
+		FSlatePostProcessSimpleBlurPassInputs BlurInputs;
+		BlurInputs.InputTexture = InputTexture;
+		BlurInputs.OutputTexture = OutputTexture;
+		BlurInputs.Strength = GaussianBlurStrength_RenderThread;
+
+		AddSlatePostProcessBlurPass(GraphBuilder, BlurInputs);
+	}
+}
+
+void FSlatePostBufferBlurProxy::OnUpdateValuesRenderThread()
+{
+	// Don't issue multiple updates in a single frame from the CPU based on dirty values
+	if (!ParamUpdateFence.IsFenceComplete())
+	{
+		return;
+	}
+
+	// Only issue an update when parent exists & values are different
+	if (USlatePostBufferBlur* ParentBlurObject = Cast<USlatePostBufferBlur>(ParentObject))
+	{
+		if (ParentBlurObject->GaussianBlurStrength != GaussianBlurStrength_RenderThread)
+		{
+			// Blur strengths can be updated from renderthread during draw or gamethread,
+			// if our parent object value matches the predraw then don't update the renderthread value.
+			// Instead we need to update our parent object's value to match the last value from renderthread
+			bool bUpdatedInRenderThread = ParentBlurObject->GaussianBlurStrength == GaussianBlurStrengthPreDraw;
+
+			if (bUpdatedInRenderThread)
+			{
+				ParentBlurObject->GaussianBlurStrength = GaussianBlurStrength_RenderThread;
+				GaussianBlurStrengthPreDraw = GaussianBlurStrength_RenderThread;
+			}
+			else
+			{
+				// Explicit param copy to avoid renderthread from reading value during gamethread write
+				float GaussianBlurStrengthCopy = ParentBlurObject->GaussianBlurStrength;
+				GaussianBlurStrengthPreDraw = GaussianBlurStrengthCopy;
+
+				// Execute param copy in a render command to safely update value on renderthread without race conditions
+				TWeakPtr<FSlatePostBufferBlurProxy> TempWeakThis = SharedThis(this);
+				ENQUEUE_RENDER_COMMAND(FUpdateValuesRenderThreadFX_Blur)([TempWeakThis, GaussianBlurStrengthCopy](FRHICommandListImmediate& RHICmdList)
+				{
+					if (TSharedPtr<FSlatePostBufferBlurProxy> SharedThisPin = TempWeakThis.Pin())
+					{
+						SharedThisPin->GaussianBlurStrength_RenderThread = GaussianBlurStrengthCopy;
+					}
+				});
+
+				// Issue fence to prevent multiple updates in a single frame
+				ParamUpdateFence.BeginFence();
+			}
+		}
+	}
+}
+
+/////////////////////////////////////////////////////
+// USlatePostBufferBlur
+
+USlatePostBufferBlur::USlatePostBufferBlur()
+{
+	RenderThreadProxy = nullptr;
+}
+
+USlatePostBufferBlur::~USlatePostBufferBlur()
+{
+	RenderThreadProxy = nullptr;
+}
+
+TSharedPtr<FSlateRHIPostBufferProcessorProxy> USlatePostBufferBlur::GetRenderThreadProxy()
+{
+	if (!RenderThreadProxy && IsInGameThread())
+	{
+		// Create a RT proxy specific for doing blurs
+		RenderThreadProxy = MakeShared<FSlatePostBufferBlurProxy>();
+		RenderThreadProxy->SetOwningProcessorObject(this);
+	}
+	return RenderThreadProxy;
+}
\ No newline at end of file
diff --git a/Engine/Source/Runtime/SlateRHIRenderer/Private/FX/SlateRHIPostBufferBlur.cpp b/Engine/Source/Runtime/SlateRHIRenderer/Private/FX/SlateRHIPostBufferBlur.cpp
deleted file mode 100644
diff --git a/Engine/Source/Runtime/SlateRHIRenderer/Private/FX/SlateRHIPostBufferProcessor.cpp b/Engine/Source/Runtime/SlateRHIRenderer/Private/FX/SlateRHIPostBufferProcessor.cpp
new file mode 100644
--- /dev/null
+++ b/Engine/Source/Runtime/SlateRHIRenderer/Private/FX/SlateRHIPostBufferProcessor.cpp
@@ -0,0 +1,10 @@
+// Copyright Epic Games, Inc. All Rights Reserved.
+
+#include "FX/SlateRHIPostBufferProcessor.h"
+
+#include "RHIResources.h"
+#include "SlateRHIRenderer.h"
+#include "StereoRendering.h"
+#include "UnrealEngine.h"
+
+#include UE_INLINE_GENERATED_CPP_BY_NAME(SlateRHIPostBufferProcessor)
\ No newline at end of file
diff --git a/Engine/Source/Runtime/SlateRHIRenderer/Private/SlateFontBlueprintLibrary.h b/Engine/Source/Runtime/SlateRHIRenderer/Private/SlateFontBlueprintLibrary.h
--- a/Engine/Source/Runtime/SlateRHIRenderer/Private/SlateFontBlueprintLibrary.h
+++ b/Engine/Source/Runtime/SlateRHIRenderer/Private/SlateFontBlueprintLibrary.h
@@ -13,6 +13,6 @@ class USlateFontBlueprintLibrary : public UObject
 {
 	GENERATED_BODY()
 
-	UFUNCTION(BlueprintPure, Category = "Slate Font", meta = (DeprecatedFunction, DeprecationMessage = "MakeSlateFontInfo from SlateFontBlueprintLibrary has been deprecated, use MakeSlateFontInfo from SlateFontInfoBlueprintLibrary instead."))
+	UFUNCTION(BlueprintPure, Category = "Slate Font", meta = (NativeMakeFunc))
 	static SLATERHIRENDERER_API FSlateFontInfo MakeSlateFontInfo(const UObject* FontObject, UObject* FontMaterial, FFontOutlineSettings OutlineSettings, FName TypefaceFontName, float Size = 24.0f, int32 LetterSpacing = 0, float SkewAmount = 0.0f, bool bForceMonospaced = false, bool bMaterialIsStencil = false, float MonospacedWidth = 1.0f);
 };
\ No newline at end of file
diff --git a/Engine/Source/Runtime/SlateRHIRenderer/Private/SlateRHIRenderer.h b/Engine/Source/Runtime/SlateRHIRenderer/Private/SlateRHIRenderer.h
--- a/Engine/Source/Runtime/SlateRHIRenderer/Private/SlateRHIRenderer.h
+++ b/Engine/Source/Runtime/SlateRHIRenderer/Private/SlateRHIRenderer.h
@@ -18,8 +18,14 @@
 
 class FSlateElementBatcher;
 class FSlateRHIRenderingPolicy;
+class FSlateRHIRenderingPolicyInterface;
+class FSlateRHIPostBufferProcessorProxy;
+class USlateRHIPostBufferProcessor;
+class USlateRHIRendererSettings;
 class ISlateStyle;
 class SWindow;
+struct FSlatePostSettings;
+struct FSlateDrawWindowsCommand;
 
 extern FMatrix CreateSlateProjectionMatrix(uint32 Width, uint32 Height);
 
diff --git a/Engine/Source/Runtime/SlateRHIRenderer/Public/FX/SlateFXSubsystem.h b/Engine/Source/Runtime/SlateRHIRenderer/Public/FX/SlateFXSubsystem.h
--- a/Engine/Source/Runtime/SlateRHIRenderer/Public/FX/SlateFXSubsystem.h
+++ b/Engine/Source/Runtime/SlateRHIRenderer/Public/FX/SlateFXSubsystem.h
@@ -2,42 +2,61 @@
 
 #pragma once
 
-#include "FX/SlateFXBaseSubsystem.h"
+#include "Subsystems/EngineSubsystem.h"
+
+#include "Engine/World.h"
 
 #include "SlateFXSubsystem.generated.h"
 
 #define UE_API SLATERHIRENDERER_API
 
 class FSlateRHIPostBufferProcessorProxy;
-class USlatePostBufferProcessor;
+class USlateRHIPostBufferProcessor;
 
-UCLASS(MinimalAPI, DisplayName = "Slate FX Subsystem (RHI)")
-class USlateFXSubsystem : public USlateFXBaseSubsystem
+UCLASS(MinimalAPI, DisplayName = "Slate FX Subsystem")
+class USlateFXSubsystem : public UEngineSubsystem
 {
 	GENERATED_BODY()
 
 public:
 
+	static UE_API USlateRHIPostBufferProcessor* GetPostProcessor(ESlatePostRT InSlatePostBufferBit);
 	static UE_API TSharedPtr<FSlateRHIPostBufferProcessorProxy> GetPostProcessorProxy(ESlatePostRT InSlatePostBufferBit);
 
 	//~ Begin UObject Interface.
 	UE_API virtual void BeginDestroy() override;
 	//~ End UObject Interface.
 
+	//~Begin UGameInstanceSubsystem Interface
+	UE_API virtual bool ShouldCreateSubsystem(UObject* Outer) const override;
+	UE_API virtual void Initialize(FSubsystemCollectionBase& Collection) override;
+	UE_API virtual void Deinitialize() override;
+	//~End UGameInstanceSubsystem Interface
+
 	/** Get post processor proxy for a particular post buffer index, if it exists */
 	UE_API TSharedPtr<FSlateRHIPostBufferProcessorProxy> GetSlatePostProcessorProxy(ESlatePostRT InPostBufferBit);
 
+public:
+
 	/** Get post processor for a particular post buffer index, if it exists */
-	UE_API USlatePostBufferProcessor* GetSlatePostProcessor(ESlatePostRT InPostBufferBit);
+	UFUNCTION(BlueprintCallable, Category = "SlateFX")
+	UE_API USlateRHIPostBufferProcessor* GetSlatePostProcessor(ESlatePostRT InPostBufferBit);
 
-protected:
+private:
 
-	virtual void OnInitProcessors() override;
-	virtual void OnCleanupProcessors() override;
+	/** Map of post RT buffer index to buffer processors, if they exist */
+	UPROPERTY(Transient)
+	TMap<ESlatePostRT, TObjectPtr<USlateRHIPostBufferProcessor>> SlatePostBufferProcessors;
 
 private:
 
-	static TSharedPtr<FSlateRHIPostBufferProcessorProxy> CreatePostBufferProcessorProxy(TSubclassOf<USlatePostBufferProcessor> PostProcessorClass);
+	/** Callback to create processors on world init */
+	void OnPreWorldInitialization(UWorld* World, const UWorld::InitializationValues IVS);
+
+	/** Callback to remove processors on world cleanup */
+	void OnPostWorldCleanup(UWorld* World, bool SessionEnded, bool bCleanupResources);
+
+private:
 
 	/** Map of post RT buffer index to buffer processor renderthread proxies, if they exist */
 	TMap<ESlatePostRT, TSharedPtr<FSlateRHIPostBufferProcessorProxy>> SlatePostBufferProcessorProxies;
diff --git a/Engine/Source/Runtime/SlateRHIRenderer/Public/FX/SlatePostBufferBlur.h b/Engine/Source/Runtime/SlateRHIRenderer/Public/FX/SlatePostBufferBlur.h
new file mode 100644
--- /dev/null
+++ b/Engine/Source/Runtime/SlateRHIRenderer/Public/FX/SlatePostBufferBlur.h
@@ -0,0 +1,68 @@
+// Copyright Epic Games, Inc. All Rights Reserved.
+
+#pragma once
+
+#include "FX/SlateRHIPostBufferProcessor.h"
+#include "SlatePostBufferBlur.generated.h"
+
+#define UE_API SLATERHIRENDERER_API
+
+/**
+ * Proxy for post buffer processor that the renderthread uses to perform processing
+ * This proxy exists because generally speaking usage on UObjects on the renderthread
+ * is a race condition due to UObjects being managed / updated by the game thread
+ */
+class FSlatePostBufferBlurProxy : public FSlateRHIPostBufferProcessorProxy
+{
+
+public:
+
+	//~ Begin FSlateRHIPostBufferProcessorProxy Interface
+	UE_API virtual void PostProcess_Renderthread(FRDGBuilder& GraphBuilder, const FScreenPassTexture& InputTexture, const FScreenPassTexture& OutputTexture) override;
+	UE_API virtual void OnUpdateValuesRenderThread() override;
+	//~ End FSlateRHIPostBufferProcessorProxy Interface
+
+	/** Blur strength to use when processing, renderthread version actually used to draw. Must be updated via render command except during initialization. */
+	float GaussianBlurStrength_RenderThread = 10;
+
+	/** 
+	 * Blur strength can be updated from both renderthread during draw and gamethread update. 
+	 * Store the last value gamethread provided so we know if we should use the renderthread value or gamethread value. 
+	 * We will use the most recently updated one.
+	 */
+	float GaussianBlurStrengthPreDraw = 10;
+
+protected:
+
+	/** Fence to allow for us to queue only one update per draw command from the gamethread */
+	FRenderCommandFence ParamUpdateFence;
+};
+
+/**
+ * Slate Post Buffer Processor that performs a simple gaussian blur to the backbuffer
+ * 
+ * Create a new asset deriving from this class to use / modify settings.
+ */
+UCLASS(MinimalAPI, Abstract, Blueprintable, CollapseCategories)
+class USlatePostBufferBlur : public USlateRHIPostBufferProcessor
+{
+	GENERATED_BODY()
+
+public:
+
+	UPROPERTY(interp, BlueprintReadWrite, Category = "GaussianBlur")
+	float GaussianBlurStrength = 10;
+
+public:
+
+	UE_API USlatePostBufferBlur();
+	UE_API virtual ~USlatePostBufferBlur() override;
+
+	UE_API virtual TSharedPtr<FSlateRHIPostBufferProcessorProxy> GetRenderThreadProxy();
+
+private:
+
+	TSharedPtr<FSlateRHIPostBufferProcessorProxy> RenderThreadProxy;
+};
+
+#undef UE_API
diff --git a/Engine/Source/Runtime/SlateRHIRenderer/Public/FX/SlateRHIPostBufferBlur.h b/Engine/Source/Runtime/SlateRHIRenderer/Public/FX/SlateRHIPostBufferBlur.h
deleted file mode 100644
diff --git a/Engine/Source/Runtime/SlateRHIRenderer/Public/FX/SlateRHIPostBufferProcessor.h b/Engine/Source/Runtime/SlateRHIRenderer/Public/FX/SlateRHIPostBufferProcessor.h
--- a/Engine/Source/Runtime/SlateRHIRenderer/Public/FX/SlateRHIPostBufferProcessor.h
+++ b/Engine/Source/Runtime/SlateRHIRenderer/Public/FX/SlateRHIPostBufferProcessor.h
@@ -3,7 +3,6 @@
 #pragma once
 
 #include "Engine/TextureRenderTarget2D.h"
-#include "FX/SlatePostBufferProcessor.h"
 #include "Layout/PaintGeometry.h"
 #include "Rendering/RenderingCommon.h"
 #include "Rendering/SlateRenderer.h"
@@ -12,6 +11,8 @@
 #include "Interfaces/SlateRHIRenderingPolicyInterface.h"
 #include "ScreenPass.h"
 
+#include "SlateRHIPostBufferProcessor.generated.h"
+
 class USlateRHIPostBufferProcessor;
 
 /**
@@ -19,7 +20,7 @@ class USlateRHIPostBufferProcessor;
  * This proxy exists because generally speaking usage on UObjects on the renderthread
  * is a race condition due to UObjects being managed / updated by the game thread
  */
-class FSlateRHIPostBufferProcessorProxy : public FSlatePostBufferProcessorProxy
+class FSlateRHIPostBufferProcessorProxy : public TSharedFromThis<FSlateRHIPostBufferProcessorProxy>
 {
 
 public:
@@ -31,6 +32,13 @@ public:
 	{
 	}
 
+	UE_DEPRECATED(5.5, "Use RDG version instead.")
+PRAGMA_DISABLE_DEPRECATION_WARNINGS
+	virtual void PostProcess_Renderthread(FRHICommandListImmediate& RHICmdList, FRHITexture* Src, FRHITexture* Dst, FIntRect SrcRect, FIntRect DstRect, FSlateRHIRenderingPolicyInterface InRenderingPolicy)
+	{
+	}
+PRAGMA_ENABLE_DEPRECATION_WARNINGS
+
 	/** 
 	 * Called when an post buffer update element is added to a renderbatch, 
 	 * gives proxies a chance to queue updates to their renderthread values based on the UObject processor.
@@ -43,7 +51,7 @@ public:
 	/**
 	 * Set the UObject that we are a renderthread proxy for, useful for doing gamethread updates from the proxy
 	 */
-	void SetOwningProcessorObject(USlatePostBufferProcessor* InParentObject)
+	void SetOwningProcessorObject(USlateRHIPostBufferProcessor* InParentObject)
 	{
 		ParentObject = InParentObject;
 	}
@@ -51,22 +59,43 @@ public:
 protected:
 
 	/** Pointer to processor that we are a proxy for, external design constraints should ensure that this is always valid */
-	TWeakObjectPtr<USlatePostBufferProcessor> ParentObject;
+	TWeakObjectPtr<USlateRHIPostBufferProcessor> ParentObject;
 };
 
 /**
  * Base class for types that can process the backbuffer scene into the slate post buffer.
  * 
- * You need to create a renderthread proxy that derives from 'FSlateRHIPostBufferProcessorProxy'
+ * Implement 'PostProcess' in your derived class. Additionally, you need to create a renderthread proxy that derives from 'FSlateRHIPostBufferProcessorProxy'
  * For an example see: USlatePostBufferBlur.
  */
-class UE_DEPRECATED(5.8, "Use USlatePostBufferProcessor instead") USlateRHIPostBufferProcessor : public UObject
+UCLASS(MinimalAPI, Abstract, Blueprintable, CollapseCategories)
+class USlateRHIPostBufferProcessor : public UObject
 {
+	GENERATED_BODY()
 
 public:
 
 	virtual ~USlateRHIPostBufferProcessor() = default;
 
+	/**
+	 * Overridable postprocess for the given source scene backbuffer provided in 'Src' into 'Dst'
+	 * You must override this method. In your override, you should copy params before executing 'ENQUEUE_RENDER_COMMAND'.
+	 * This allows you to avoid render & game thread race conditions. See 'SlatePostBufferBlur' for example.
+	 * Also avoid capturing [this] in your override, to avoid possible GC issues with the processor instance.
+	 *
+	 * @param InViewInfo			'FViewportInfo' resource used to get backbuffer in standalone
+	 * @param InViewportTexture		'FSlateRenderTargetRHI' resource used to get the 'BufferedRT' viewport texture used in PIE
+	 * @param InElementWindowSize	Size of window being rendered, used to determine if using stereo rendering or not.
+	 * @param InRenderingPolicy		Slate RHI RenderingPolicy
+	 * @param InSlatePostBuffer		Texture render target used for final output
+	 */
+	UE_DEPRECATED(5.5, "This path is longer supported. Get the render proxy instead.")
+PRAGMA_DISABLE_DEPRECATION_WARNINGS
+	virtual void PostProcess(FRenderResource* InViewInfo, FRenderResource* InViewportTexture, FVector2D InElementWindowSize, FSlateRHIRenderingPolicyInterface InRenderingPolicy, UTextureRenderTarget2D* InSlatePostBuffer)
+	{
+	}
+PRAGMA_ENABLE_DEPRECATION_WARNINGS
+
 	/**
 	 * Gets proxy for this post buffer processor, for execution on the renderthread
 	 */
diff --git a/Engine/Source/Runtime/SlateRHIRenderer/Public/SlateRHIRendererSettings.h b/Engine/Source/Runtime/SlateRHIRenderer/Public/SlateRHIRendererSettings.h
--- a/Engine/Source/Runtime/SlateRHIRenderer/Public/SlateRHIRendererSettings.h
+++ b/Engine/Source/Runtime/SlateRHIRenderer/Public/SlateRHIRendererSettings.h
@@ -15,7 +15,7 @@
 
 #define UE_API SLATERHIRENDERER_API
 
-class USlatePostBufferProcessor;
+class USlateRHIPostBufferProcessor;
 class UTextureRenderTarget2D;
 
 /**
@@ -58,7 +58,7 @@ public:
 
 	/** Copy of actually loaded post processor class */
 	UPROPERTY(EditAnywhere, BlueprintReadWrite, Category = BufferSettings, Meta = (AllowAbstract = false))
-	TSubclassOf<USlatePostBufferProcessor> PostProcessorClass;
+	TSubclassOf<USlateRHIPostBufferProcessor> PostProcessorClass;
 
 private:
 
diff --git a/Engine/Source/Runtime/SlateRHIRenderer/SlateRHIRenderer.Build.cs b/Engine/Source/Runtime/SlateRHIRenderer/SlateRHIRenderer.Build.cs
--- a/Engine/Source/Runtime/SlateRHIRenderer/SlateRHIRenderer.Build.cs
+++ b/Engine/Source/Runtime/SlateRHIRenderer/SlateRHIRenderer.Build.cs
@@ -6,12 +6,6 @@ public class SlateRHIRenderer : ModuleRules
 {
     public SlateRHIRenderer(ReadOnlyTargetRules Target) : base(Target)
 	{
-		PublicDependencyModuleNames.AddRange(
-			new string[] {
-				"SlateBaseRenderer"
-			}
-		);
-
 		PrivateDependencyModuleNames.AddRange(
 			new string[] {
 				"Core",
diff --git a/Engine/Source/Runtime/UMG/Private/Components/PostBufferBlurUpdater.cpp b/Engine/Source/Runtime/UMG/Private/Components/PostBufferBlurUpdater.cpp
--- a/Engine/Source/Runtime/UMG/Private/Components/PostBufferBlurUpdater.cpp
+++ b/Engine/Source/Runtime/UMG/Private/Components/PostBufferBlurUpdater.cpp
@@ -5,7 +5,7 @@
 // Exclude SlateRHIRenderer related includes & implementations from server builds since the module is not a dependency / will not link for UMG on the server
 #if !UE_SERVER
 #include "FX/SlateRHIPostBufferProcessor.h"
-#include "FX/SlateRHIPostBufferBlur.h"
+#include "FX/SlatePostBufferBlur.h"
 #endif // !UE_SERVER
 
 #include UE_INLINE_GENERATED_CPP_BY_NAME(PostBufferBlurUpdater)
